Initial network replay streaming interface

diff --git a/Engine/Source/Runtime/NetworkReplayStreaming/NetworkReplayStreaming/NetworkReplayStreaming.Build.cs b/Engine/Source/Runtime/NetworkReplayStreaming/NetworkReplayStreaming/NetworkReplayStreaming.Build.cs
new file mode 100644
--- /dev/null
+++ b/Engine/Source/Runtime/NetworkReplayStreaming/NetworkReplayStreaming/NetworkReplayStreaming.Build.cs
@@ -0,0 +1,20 @@
+// Copyright 1998-2015 Epic Games, Inc. All Rights Reserved.
+
+namespace UnrealBuildTool.Rules
+{
+	public class NetworkReplayStreaming : ModuleRules
+	{
+		public NetworkReplayStreaming( TargetInfo Target )
+		{
+			PrivateIncludePaths.Add( "Runtime/NetworkReplayStreaming/NetworkReplayStreaming/Private" );
+
+			PublicDependencyModuleNames.AddRange(
+				new string[]
+				{
+					"Core",
+                    // ... add other public dependencies that you statically link with here ...
+				}
+			);
+		}
+	}
+}
\ No newline at end of file
diff --git a/Engine/Source/Runtime/NetworkReplayStreaming/NetworkReplayStreaming/Private/NetworkReplayStreaming.cpp b/Engine/Source/Runtime/NetworkReplayStreaming/NetworkReplayStreaming/Private/NetworkReplayStreaming.cpp
new file mode 100644
--- /dev/null
+++ b/Engine/Source/Runtime/NetworkReplayStreaming/NetworkReplayStreaming/Private/NetworkReplayStreaming.cpp
@@ -0,0 +1,24 @@
+// Copyright 1998-2015 Epic Games, Inc. All Rights Reserved.
+
+#include "NetworkReplayStreaming.h"
+
+IMPLEMENT_MODULE( FNetworkReplayStreaming, NetworkReplayStreaming );
+
+INetworkReplayStreamingFactory& FNetworkReplayStreaming::GetFactory()
+{
+	FString FactoryName = TEXT( "NullNetworkReplayStreaming" );
+	GConfig->GetString( TEXT( "NetworkReplayStreaming" ), TEXT( "DefaultFactoryName" ), FactoryName, GEngineIni );
+
+	// See if we need to forcefully fallback to the null streamer
+	if ( !FModuleManager::Get().IsModuleLoaded( *FactoryName ) )
+	{
+		FModuleManager::Get().LoadModule( *FactoryName );
+	
+		if ( !FModuleManager::Get().IsModuleLoaded( *FactoryName ) )
+		{
+			FactoryName = TEXT( "NullNetworkReplayStreaming" );
+		}
+	}
+
+	return FModuleManager::Get().LoadModuleChecked< INetworkReplayStreamingFactory >( *FactoryName );
+}
\ No newline at end of file
diff --git a/Engine/Source/Runtime/NetworkReplayStreaming/NetworkReplayStreaming/Public/NetworkReplayStreaming.h b/Engine/Source/Runtime/NetworkReplayStreaming/NetworkReplayStreaming/Public/NetworkReplayStreaming.h
new file mode 100644
--- /dev/null
+++ b/Engine/Source/Runtime/NetworkReplayStreaming/NetworkReplayStreaming/Public/NetworkReplayStreaming.h
@@ -0,0 +1,40 @@
+// Copyright 1998-2015 Epic Games, Inc. All Rights Reserved.
+
+#pragma once
+
+// Dependencies.
+#include "ModuleManager.h"
+#include "Core.h"
+
+DECLARE_MULTICAST_DELEGATE_TwoParams( FOnStreamReady, bool, bool );
+typedef FOnStreamReady::FDelegate FOnStreamReadyDelegate;
+
+/** Generic interface for network replay streaming */
+class INetworkReplayStreamer 
+{
+public:
+	virtual ~INetworkReplayStreamer() {}
+
+	virtual void StartStreaming( FString& StreamName, bool bRecord, const FOnStreamReadyDelegate& Delegate ) = 0;
+	virtual void StopStreaming() = 0;
+	virtual FArchive* GetStreamingArchive() = 0;
+};
+
+/** Replay streamer factory */
+class INetworkReplayStreamingFactory : public IModuleInterface
+{
+public:
+	virtual TSharedPtr< INetworkReplayStreamer > CreateReplayStreamer() = 0;
+};
+
+/** Replay streaming factory manager */
+class FNetworkReplayStreaming : public IModuleInterface
+{
+public:
+	static inline FNetworkReplayStreaming& Get()
+	{
+		return FModuleManager::LoadModuleChecked< FNetworkReplayStreaming >( "NetworkReplayStreaming" );
+	}
+
+	virtual INetworkReplayStreamingFactory& GetFactory();
+};
