Initial reference/null network replay streaming implementation

diff --git a/Engine/Source/Runtime/NetworkReplayStreaming/NullNetworkReplayStreaming/NullNetworkReplayStreaming.build.cs b/Engine/Source/Runtime/NetworkReplayStreaming/NullNetworkReplayStreaming/NullNetworkReplayStreaming.build.cs
new file mode 100644
--- /dev/null
+++ b/Engine/Source/Runtime/NetworkReplayStreaming/NullNetworkReplayStreaming/NullNetworkReplayStreaming.build.cs
@@ -0,0 +1,20 @@
+// Copyright 1998-2015 Epic Games, Inc. All Rights Reserved.
+
+namespace UnrealBuildTool.Rules
+{
+	public class NullNetworkReplayStreaming : ModuleRules
+	{
+		public NullNetworkReplayStreaming( TargetInfo Target )
+		{
+			PrivateIncludePaths.Add( "Runtime/NetworkReplayStreaming/NullNetworkReplayStreaming/Private" );
+
+			PublicDependencyModuleNames.AddRange(
+				new string[]
+				{
+					"Core",
+					"NetworkReplayStreaming",
+				}
+			);
+		}
+	}
+}
\ No newline at end of file
diff --git a/Engine/Source/Runtime/NetworkReplayStreaming/NullNetworkReplayStreaming/Private/NullNetworkReplayStreaming.cpp b/Engine/Source/Runtime/NetworkReplayStreaming/NullNetworkReplayStreaming/Private/NullNetworkReplayStreaming.cpp
new file mode 100644
--- /dev/null
+++ b/Engine/Source/Runtime/NetworkReplayStreaming/NullNetworkReplayStreaming/Private/NullNetworkReplayStreaming.cpp
@@ -0,0 +1,45 @@
+// Copyright 1998-2015 Epic Games, Inc. All Rights Reserved.
+
+#include "NullNetworkReplayStreaming.h"
+
+/**
+ * Very basic implementation of network replay streaming using the file system
+ * As of now, there is just simple opening and closing of the stream, and handing off the stream for direct use
+ * Eventually, we'll want to expand this interface to allow enumerating demos, support for live spectating on local machine
+ * (which will require support for writing/reading at the same time)
+ */
+
+void FNullNetworkReplayStreamer::StartStreaming( FString& StreamName, bool bRecord, const FOnStreamReadyDelegate& Delegate )
+{
+	if ( !bRecord )
+	{
+		// Opne file for reading
+		FileAr = IFileManager::Get().CreateFileReader( *StreamName );
+	}
+	else
+	{
+		// Open file for writing
+		FileAr = IFileManager::Get().CreateFileWriter( *StreamName );
+	}
+
+	// Notify immediately
+	Delegate.ExecuteIfBound( FileAr != NULL, bRecord );
+}
+
+void FNullNetworkReplayStreamer::StopStreaming()
+{
+	delete FileAr;
+	FileAr = NULL;
+}
+
+FArchive* FNullNetworkReplayStreamer::GetStreamingArchive()
+{
+	return FileAr;
+}
+
+IMPLEMENT_MODULE( FNullNetworkReplayStreamingFactory, NullNetworkReplayStreaming )
+
+TSharedPtr< INetworkReplayStreamer > FNullNetworkReplayStreamingFactory::CreateReplayStreamer() 
+{
+	return TSharedPtr< INetworkReplayStreamer >( new FNullNetworkReplayStreamer );
+}
diff --git a/Engine/Source/Runtime/NetworkReplayStreaming/NullNetworkReplayStreaming/Public/NullNetworkReplayStreaming.h b/Engine/Source/Runtime/NetworkReplayStreaming/NullNetworkReplayStreaming/Public/NullNetworkReplayStreaming.h
new file mode 100644
--- /dev/null
+++ b/Engine/Source/Runtime/NetworkReplayStreaming/NullNetworkReplayStreaming/Public/NullNetworkReplayStreaming.h
@@ -0,0 +1,27 @@
+// Copyright 1998-2015 Epic Games, Inc. All Rights Reserved.
+
+#pragma once
+
+#include "NetworkReplayStreaming.h"
+#include "Core.h"
+#include "ModuleManager.h"
+
+/** Default streamer that goes straight to the HD */
+class FNullNetworkReplayStreamer : public INetworkReplayStreamer
+{
+public:
+	FNullNetworkReplayStreamer() : FileAr( NULL ) {}
+	virtual void StartStreaming( FString& StreamName, bool bRecord, const FOnStreamReadyDelegate& Delegate ) override;
+	virtual void StopStreaming() override;
+	virtual FArchive* GetStreamingArchive() override;
+
+private:
+	/** Handle to the archive that will read/write network packets */
+	FArchive* FileAr;
+};
+
+class FNullNetworkReplayStreamingFactory : public INetworkReplayStreamingFactory
+{
+public:
+	virtual TSharedPtr< INetworkReplayStreamer > CreateReplayStreamer();
+};
