Initial http network replay streaming implementation

diff --git a/Engine/Source/Runtime/NetworkReplayStreaming/HttpNetworkReplayStreaming/HttpNetworkReplayStreaming.Build.cs b/Engine/Source/Runtime/NetworkReplayStreaming/HttpNetworkReplayStreaming/HttpNetworkReplayStreaming.Build.cs
new file mode 100644
--- /dev/null
+++ b/Engine/Source/Runtime/NetworkReplayStreaming/HttpNetworkReplayStreaming/HttpNetworkReplayStreaming.Build.cs
@@ -0,0 +1,22 @@
+// Copyright 1998-2015 Epic Games, Inc. All Rights Reserved.
+
+namespace UnrealBuildTool.Rules
+{
+	public class HttpNetworkReplayStreaming : ModuleRules
+    {
+        public HttpNetworkReplayStreaming(TargetInfo Target)
+        {
+			PrivateIncludePaths.Add( "Runtime/NetworkReplayStreaming/HttpNetworkReplayStreaming/Private" );
+
+			PrivateDependencyModuleNames.AddRange(
+                new string[]
+				{
+					"Core",
+					"CoreUObject",
+                    "Engine",
+                    "HTTP",
+					"NetworkReplayStreaming"
+				});
+        }
+    }
+}
diff --git a/Engine/Source/Runtime/NetworkReplayStreaming/HttpNetworkReplayStreaming/Private/HttpNetworkReplayStreaming.cpp b/Engine/Source/Runtime/NetworkReplayStreaming/HttpNetworkReplayStreaming/Private/HttpNetworkReplayStreaming.cpp
new file mode 100644
--- /dev/null
+++ b/Engine/Source/Runtime/NetworkReplayStreaming/HttpNetworkReplayStreaming/Private/HttpNetworkReplayStreaming.cpp
@@ -0,0 +1,145 @@
+// Copyright 1998-2015 Epic Games, Inc. All Rights Reserved.
+
+#include "HttpNetworkReplayStreaming.h"
+
+void HttpStreamFArchive::Serialize( void* V, int64 Length ) 
+{
+	if ( IsLoading() )
+	{
+		if ( Pos + Length > Buffer.Num() )
+		{
+			ArIsError = true;
+			return;
+		}
+			
+		FMemory::Memcpy( V, Buffer.GetData() + Pos, Length );
+
+		Pos += Length;
+	}
+	else
+	{
+		check( Pos <= Buffer.Num() );
+
+		const int32 SpaceNeeded = Length - ( Buffer.Num() - Pos );
+
+		if ( SpaceNeeded > 0 )
+		{
+			Buffer.AddZeroed( SpaceNeeded );
+		}
+
+		FMemory::Memcpy( Buffer.GetData() + Pos, V, Length );
+
+		Pos += Length;
+	}
+}
+
+int64 HttpStreamFArchive::Tell() 
+{
+	return Pos;
+}
+
+int64 HttpStreamFArchive::TotalSize()
+{
+	return Buffer.Num();
+}
+
+void HttpStreamFArchive::Seek( int64 InPos ) 
+{
+	Pos = InPos;
+}
+
+void FHttpNetworkReplayStreamer::StartStreaming( FString& StreamName, bool bRecord, const FOnStreamReadyDelegate& Delegate )
+{
+	StopStreaming();
+
+	RememberedDelegate = Delegate;
+
+	Archive.ArIsLoading = !bRecord;
+	Archive.ArIsSaving = !Archive.ArIsLoading;
+
+	// We need to chop off any long path, we just want the name
+	// FIXME: Change it so that by convention, the caller uses the short name
+	// Individual streamers can add whatever path they like
+	DemoShortName = StreamName;
+
+	int32 Index = 0;
+
+	if ( StreamName.FindLastChar( '/', Index ) )
+	{
+		DemoShortName = StreamName.RightChop( Index + 1 );
+	}
+
+	// If we're recording, we'll send the stream after we're done
+	// NOTE - We're not really streaming right now, but this is a start
+	if ( Archive.ArIsLoading )
+	{
+		// Create the Http request and add to pending request list
+		TSharedRef<class IHttpRequest> HttpRequest = FHttpModule::Get().CreateRequest();
+
+		HttpRequest->OnProcessRequestComplete().BindRaw( this, &FHttpNetworkReplayStreamer::HttpRequestFinished );
+
+		FString ServerURL;
+		GConfig->GetString( TEXT( "HttpNetworkReplayStreaming" ), TEXT( "ServerURL" ), ServerURL, GEngineIni );
+
+		HttpRequest->SetURL( ServerURL + TEXT( "download/" ) + DemoShortName );
+		HttpRequest->SetVerb( TEXT( "GET" ) );
+		HttpRequest->ProcessRequest();
+	}
+}
+
+void FHttpNetworkReplayStreamer::StopStreaming()
+{
+	if ( Archive.ArIsSaving && Archive.Buffer.Num() > 0 )
+	{
+		// Create the Http request and add to pending request list
+		TSharedRef<class IHttpRequest> HttpRequest = FHttpModule::Get().CreateRequest();
+
+		// Currently, the engine free's this streamer as soon as it calls this, so we can't register for a callback here right now
+		// This is ok anyhow, since uploading is fire and foreget (we aren't actually streaming it yet)
+		//HttpRequest->OnProcessRequestComplete().BindRaw( this, &FHttpNetworkReplayStreamer::HttpRequestFinished );
+
+		FString ServerURL;
+		GConfig->GetString( TEXT( "HttpNetworkReplayStreaming" ), TEXT( "ServerURL" ), ServerURL, GEngineIni );
+
+		HttpRequest->SetURL( ServerURL + TEXT( "upload/" ) + DemoShortName );
+		HttpRequest->SetVerb( TEXT( "POST" ) );
+		HttpRequest->SetHeader( TEXT( "Content-Type" ), TEXT( "application/octet-stream" ) );
+		HttpRequest->SetContent( Archive.Buffer );
+		HttpRequest->ProcessRequest();
+	}
+
+	Archive.ArIsLoading = false;
+	Archive.ArIsSaving = false;
+	Archive.Buffer.Empty();
+	Archive.Pos = 0;
+}
+
+FArchive* FHttpNetworkReplayStreamer::GetStreamingArchive()
+{
+	return ( Archive.IsSaving() || Archive.Buffer.Num() >  0 ) ? &Archive : NULL;
+}
+
+void FHttpNetworkReplayStreamer::HttpRequestFinished( FHttpRequestPtr HttpRequest, FHttpResponsePtr HttpResponse, bool bSucceeded )
+{
+	if ( bSucceeded )
+	{
+		if ( Archive.IsLoading() )
+		{
+			Archive.Buffer = HttpResponse->GetContent();
+		}
+	}
+	else
+	{
+		Archive.Buffer.Empty();
+	}
+
+	RememberedDelegate.ExecuteIfBound( bSucceeded, Archive.IsSaving() );
+	RememberedDelegate = FOnStreamReadyDelegate();
+}
+
+IMPLEMENT_MODULE( FHttpNetworkReplayStreamingFactory, HttpNetworkReplayStreaming )
+
+TSharedPtr< INetworkReplayStreamer > FHttpNetworkReplayStreamingFactory::CreateReplayStreamer()
+{
+	return TSharedPtr< INetworkReplayStreamer >( new FHttpNetworkReplayStreamer );
+}
diff --git a/Engine/Source/Runtime/NetworkReplayStreaming/HttpNetworkReplayStreaming/Private/HttpNetworkReplayStreaming.h b/Engine/Source/Runtime/NetworkReplayStreaming/HttpNetworkReplayStreaming/Private/HttpNetworkReplayStreaming.h
new file mode 100644
--- /dev/null
+++ b/Engine/Source/Runtime/NetworkReplayStreaming/HttpNetworkReplayStreaming/Private/HttpNetworkReplayStreaming.h
@@ -0,0 +1,47 @@
+// Copyright 1998-2015 Epic Games, Inc. All Rights Reserved.
+#pragma once
+
+#include "NetworkReplayStreaming.h"
+#include "HTTP.h"
+
+/**
+ * Archive used to buffer stream over http
+ */
+class HttpStreamFArchive : public FArchive
+{
+public:
+	HttpStreamFArchive() : Pos( 0 ) {}
+
+	virtual void	Serialize( void* V, int64 Length );
+	virtual int64	Tell();
+	virtual int64	TotalSize();
+	virtual void	Seek( int64 InPos );
+
+	TArray< uint8 >	Buffer;
+	int32			Pos;
+};
+
+/**
+ * Http network replay streaming manager
+ */
+class FHttpNetworkReplayStreamer : public INetworkReplayStreamer
+{
+public:
+	/** IHttpNetworkReplayStreamer implementation */
+	FHttpNetworkReplayStreamer() {}
+	virtual void StartStreaming( FString& StreamName, bool bRecord, const FOnStreamReadyDelegate& Delegate ) override;
+	virtual void StopStreaming() override;
+	virtual FArchive* GetStreamingArchive() override;
+
+	void HttpRequestFinished( FHttpRequestPtr HttpRequest, FHttpResponsePtr HttpResponse, bool bSucceeded );
+
+	FOnStreamReadyDelegate	RememberedDelegate;
+	HttpStreamFArchive		Archive;
+	FString					DemoShortName;
+};
+
+class FHttpNetworkReplayStreamingFactory : public INetworkReplayStreamingFactory
+{
+public:
+	virtual TSharedPtr< INetworkReplayStreamer > CreateReplayStreamer();
+};
