Changes necessary now that replay streamer is in charge of path

diff --git a/Engine/Source/Runtime/NetworkReplayStreaming/HttpNetworkReplayStreaming/Private/HttpNetworkReplayStreaming.cpp b/Engine/Source/Runtime/NetworkReplayStreaming/HttpNetworkReplayStreaming/Private/HttpNetworkReplayStreaming.cpp
--- a/Engine/Source/Runtime/NetworkReplayStreaming/HttpNetworkReplayStreaming/Private/HttpNetworkReplayStreaming.cpp
+++ b/Engine/Source/Runtime/NetworkReplayStreaming/HttpNetworkReplayStreaming/Private/HttpNetworkReplayStreaming.cpp
@@ -57,17 +57,7 @@ void FHttpNetworkReplayStreamer::StartStreaming( FString& StreamName, bool bReco
 	Archive.ArIsLoading = !bRecord;
 	Archive.ArIsSaving = !Archive.ArIsLoading;
 
-	// We need to chop off any long path, we just want the name
-	// FIXME: Change it so that by convention, the caller uses the short name
-	// Individual streamers can add whatever path they like
-	DemoShortName = StreamName;
-
-	int32 Index = 0;
-
-	if ( StreamName.FindLastChar( '/', Index ) )
-	{
-		DemoShortName = StreamName.RightChop( Index + 1 );
-	}
+	DemoShortName = StreamName + ".demo";
 
 	// If we're recording, we'll send the stream after we're done
 	// NOTE - We're not really streaming right now, but this is a start
