Changes to http replay streamer to support versions and sessions

diff --git a/Engine/Source/Runtime/NetworkReplayStreaming/HttpNetworkReplayStreaming/Private/HttpNetworkReplayStreaming.cpp b/Engine/Source/Runtime/NetworkReplayStreaming/HttpNetworkReplayStreaming/Private/HttpNetworkReplayStreaming.cpp
--- a/Engine/Source/Runtime/NetworkReplayStreaming/HttpNetworkReplayStreaming/Private/HttpNetworkReplayStreaming.cpp
+++ b/Engine/Source/Runtime/NetworkReplayStreaming/HttpNetworkReplayStreaming/Private/HttpNetworkReplayStreaming.cpp
@@ -2,6 +2,8 @@
 
 #include "HttpNetworkReplayStreaming.h"
 
+DEFINE_LOG_CATEGORY_STATIC( LogHttpReplay, Log, All );
+
 void HttpStreamFArchive::Serialize( void* V, int64 Length ) 
 {
 	if ( IsLoading() )
@@ -57,7 +59,7 @@ void FHttpNetworkReplayStreamer::StartStreaming( FString& StreamName, bool bReco
 	Archive.ArIsLoading = !bRecord;
 	Archive.ArIsSaving = !Archive.ArIsLoading;
 
-	DemoShortName = StreamName + ".demo";
+	DemoShortName = StreamName;
 
 	// If we're recording, we'll send the stream after we're done
 	// NOTE - We're not really streaming right now, but this is a start
@@ -71,7 +73,7 @@ void FHttpNetworkReplayStreamer::StartStreaming( FString& StreamName, bool bReco
 		FString ServerURL;
 		GConfig->GetString( TEXT( "HttpNetworkReplayStreaming" ), TEXT( "ServerURL" ), ServerURL, GEngineIni );
 
-		HttpRequest->SetURL( ServerURL + TEXT( "download/" ) + DemoShortName );
+		HttpRequest->SetURL( FString::Printf( TEXT( "%sdownload?Version=%s&Session=%s&Filename=%s" ), *ServerURL, TEXT( "Test" ), *DemoShortName, *( DemoShortName + TEXT( ".demo" ) ) ) );
 		HttpRequest->SetVerb( TEXT( "GET" ) );
 		HttpRequest->ProcessRequest();
 	}
@@ -91,7 +93,7 @@ void FHttpNetworkReplayStreamer::StopStreaming()
 		FString ServerURL;
 		GConfig->GetString( TEXT( "HttpNetworkReplayStreaming" ), TEXT( "ServerURL" ), ServerURL, GEngineIni );
 
-		HttpRequest->SetURL( ServerURL + TEXT( "upload/" ) + DemoShortName );
+		HttpRequest->SetURL( FString::Printf( TEXT( "%supload?Version=%s&Session=%s&Filename=%s" ), *ServerURL, TEXT( "Test" ), *DemoShortName, *( DemoShortName + TEXT( ".demo" ) ) ) );
 		HttpRequest->SetVerb( TEXT( "POST" ) );
 		HttpRequest->SetHeader( TEXT( "Content-Type" ), TEXT( "application/octet-stream" ) );
 		HttpRequest->SetContent( Archive.Buffer );
@@ -118,6 +120,9 @@ void FHttpNetworkReplayStreamer::HttpRequestFinished( FHttpRequestPtr HttpReques
 {
 	if ( bSucceeded )
 	{
+		//FString SessionId = HttpResponse->GetHeader( TEXT( "SessionId" ) );
+		//UE_LOG( LogHttpReplay, Log, TEXT( "SessionId: %s" ), *SessionId );
+
 		if ( Archive.IsLoading() )
 		{
 			Archive.Buffer = HttpResponse->GetContent();
