Client side http replay streaming support

- Official session support client side
- Initial client side support for multi-part uploads

diff --git a/Engine/Source/Runtime/NetworkReplayStreaming/HttpNetworkReplayStreaming/Private/HttpNetworkReplayStreaming.cpp b/Engine/Source/Runtime/NetworkReplayStreaming/HttpNetworkReplayStreaming/Private/HttpNetworkReplayStreaming.cpp
--- a/Engine/Source/Runtime/NetworkReplayStreaming/HttpNetworkReplayStreaming/Private/HttpNetworkReplayStreaming.cpp
+++ b/Engine/Source/Runtime/NetworkReplayStreaming/HttpNetworkReplayStreaming/Private/HttpNetworkReplayStreaming.cpp
@@ -56,59 +56,102 @@ void FHttpNetworkReplayStreamer::StartStreaming( FString& StreamName, bool bReco
 
 	RememberedDelegate = Delegate;
 
-	Archive.ArIsLoading = !bRecord;
-	Archive.ArIsSaving = !Archive.ArIsLoading;
+	StreamArchive.ArIsLoading = !bRecord;
+	StreamArchive.ArIsSaving = !StreamArchive.ArIsLoading;
 
-	DemoShortName = StreamName;
+	LastFlushTime = FPlatformTime::Seconds();
 
-	// If we're recording, we'll send the stream after we're done
-	// NOTE - We're not really streaming right now, but this is a start
-	if ( Archive.ArIsLoading )
+	const bool bOverrideRecordingSession = false;
+
+	// Override session name
+	if ( StreamArchive.ArIsLoading || bOverrideRecordingSession )
 	{
-		// Create the Http request and add to pending request list
-		TSharedRef<class IHttpRequest> HttpRequest = FHttpModule::Get().CreateRequest();
+		SessionName = StreamName;
+	}
 
-		HttpRequest->OnProcessRequestComplete().BindRaw( this, &FHttpNetworkReplayStreamer::HttpRequestFinished );
+	// Create the Http request and add to pending request list
+	TSharedRef<class IHttpRequest> HttpRequest = FHttpModule::Get().CreateRequest();
 
-		FString ServerURL;
-		GConfig->GetString( TEXT( "HttpNetworkReplayStreaming" ), TEXT( "ServerURL" ), ServerURL, GEngineIni );
+	FString ServerURL;
+	GConfig->GetString( TEXT( "HttpNetworkReplayStreaming" ), TEXT( "ServerURL" ), ServerURL, GEngineIni );
 
-		HttpRequest->SetURL( FString::Printf( TEXT( "%sdownload?Version=%s&Session=%s&Filename=%s" ), *ServerURL, TEXT( "Test" ), *DemoShortName, *( DemoShortName + TEXT( ".demo" ) ) ) );
+	if ( StreamArchive.ArIsLoading )
+	{
+		// Download the demo file from the http server
+		HttpRequest->SetURL( FString::Printf( TEXT( "%sdownload?Version=%s&Session=%s&Filename=stream.%i" ), *ServerURL, TEXT( "Test" ), *SessionName, StreamFileCount ) );
 		HttpRequest->SetVerb( TEXT( "GET" ) );
-		HttpRequest->ProcessRequest();
+
+		HttpRequest->OnProcessRequestComplete().BindRaw( this, &FHttpNetworkReplayStreamer::HttpDownloadFinished );
+
+		StreamFileCount++;
 	}
+	else
+	{
+		// Make sure this session is registered with the http server
+		if ( !SessionName.IsEmpty() )
+		{
+			HttpRequest->SetURL( FString::Printf( TEXT( "%sstartstreamingup?Version=%s&Session=%s" ), *ServerURL, TEXT( "Test" ), *SessionName ) );
+		}
+		else
+		{
+			HttpRequest->SetURL( FString::Printf( TEXT( "%sstartstreamingup?Version=%s" ), *ServerURL, TEXT( "Test" ) ) );
+		}
+
+		HttpRequest->SetVerb( TEXT( "POST" ) );
+
+		HttpRequest->OnProcessRequestComplete().BindRaw( this, &FHttpNetworkReplayStreamer::HttpStartStreamingUpFinished );
+
+		SessionName.Empty();
+	}
+	
+	HttpRequest->ProcessRequest();
 }
 
 void FHttpNetworkReplayStreamer::StopStreaming()
 {
-	if ( Archive.ArIsSaving && Archive.Buffer.Num() > 0 )
+	FlushStream();
+
+	StreamArchive.ArIsLoading = false;
+	StreamArchive.ArIsSaving = false;
+	StreamArchive.Buffer.Empty();
+	StreamArchive.Pos = 0;
+	StreamFileCount = 0;
+}
+
+void FHttpNetworkReplayStreamer::FlushStream()
+{
+	if ( StreamArchive.ArIsSaving && StreamArchive.Buffer.Num() > 0 && !SessionName.IsEmpty() )
 	{
+		UE_LOG( LogHttpReplay, Log, TEXT( "FHttpNetworkReplayStreamer::FlushStream. StreamFileCount: %i, Size: %i" ), StreamFileCount, StreamArchive.Buffer.Num() );
+
 		// Create the Http request and add to pending request list
 		TSharedRef<class IHttpRequest> HttpRequest = FHttpModule::Get().CreateRequest();
 
 		// Currently, the engine free's this streamer as soon as it calls this, so we can't register for a callback here right now
-		// This is ok anyhow, since uploading is fire and foreget (we aren't actually streaming it yet)
+		// This is ok anyhow, since uploading is fire and forget (we aren't actually streaming it yet)
 		//HttpRequest->OnProcessRequestComplete().BindRaw( this, &FHttpNetworkReplayStreamer::HttpRequestFinished );
 
 		FString ServerURL;
 		GConfig->GetString( TEXT( "HttpNetworkReplayStreaming" ), TEXT( "ServerURL" ), ServerURL, GEngineIni );
 
-		HttpRequest->SetURL( FString::Printf( TEXT( "%supload?Version=%s&Session=%s&Filename=%s" ), *ServerURL, TEXT( "Test" ), *DemoShortName, *( DemoShortName + TEXT( ".demo" ) ) ) );
+		HttpRequest->SetURL( FString::Printf( TEXT( "%supload?Version=%s&Session=%s&Filename=stream.%i" ), *ServerURL, TEXT( "Test" ), *SessionName, StreamFileCount ) );
 		HttpRequest->SetVerb( TEXT( "POST" ) );
 		HttpRequest->SetHeader( TEXT( "Content-Type" ), TEXT( "application/octet-stream" ) );
-		HttpRequest->SetContent( Archive.Buffer );
+		HttpRequest->SetContent( StreamArchive.Buffer );
 		HttpRequest->ProcessRequest();
-	}
 
-	Archive.ArIsLoading = false;
-	Archive.ArIsSaving = false;
-	Archive.Buffer.Empty();
-	Archive.Pos = 0;
+		StreamFileCount++;
+
+		StreamArchive.Buffer.Empty();
+		StreamArchive.Pos = 0;
+
+		LastFlushTime = FPlatformTime::Seconds();
+	}
 }
 
 FArchive* FHttpNetworkReplayStreamer::GetStreamingArchive()
 {
-	return ( Archive.IsSaving() || Archive.Buffer.Num() >  0 ) ? &Archive : NULL;
+	return ( StreamArchive.IsSaving() || StreamArchive.Buffer.Num() >  0 ) ? &StreamArchive : NULL;
 }
 
 FArchive* FHttpNetworkReplayStreamer::GetMetadataArchive()
@@ -116,27 +159,66 @@ FArchive* FHttpNetworkReplayStreamer::GetMetadataArchive()
 	return NULL;
 }
 
-void FHttpNetworkReplayStreamer::HttpRequestFinished( FHttpRequestPtr HttpRequest, FHttpResponsePtr HttpResponse, bool bSucceeded )
+void FHttpNetworkReplayStreamer::HttpStartStreamingUpFinished( FHttpRequestPtr HttpRequest, FHttpResponsePtr HttpResponse, bool bSucceeded )
 {
-	if ( bSucceeded )
+	if ( bSucceeded && HttpResponse->GetResponseCode() == EHttpResponseCodes::Ok )
+	{
+		SessionName = HttpResponse->GetHeader( TEXT( "Session" ) );
+
+		UE_LOG( LogHttpReplay, Log, TEXT( "FHttpNetworkReplayStreamer::HttpStartStreamingUpFinished. SessionName: %s" ), *SessionName );
+	}
+	else
 	{
-		//FString SessionId = HttpResponse->GetHeader( TEXT( "SessionId" ) );
-		//UE_LOG( LogHttpReplay, Log, TEXT( "SessionId: %s" ), *SessionId );
+		UE_LOG( LogHttpReplay, Warning, TEXT( "FHttpNetworkReplayStreamer::HttpStartStreamingUpFinished. FAILED" ), *SessionName );
+	}
+}
 
-		if ( Archive.IsLoading() )
+void FHttpNetworkReplayStreamer::HttpDownloadFinished( FHttpRequestPtr HttpRequest, FHttpResponsePtr HttpResponse, bool bSucceeded )
+{
+	if ( bSucceeded && HttpResponse->GetResponseCode() == EHttpResponseCodes::Ok )
+	{
+		if ( StreamArchive.IsLoading() )
 		{
-			Archive.Buffer = HttpResponse->GetContent();
+			StreamArchive.Buffer = HttpResponse->GetContent();
 		}
+
+		RememberedDelegate.ExecuteIfBound( true, StreamArchive.IsSaving() );
 	}
 	else
 	{
-		Archive.Buffer.Empty();
+		// FAIL
+		StreamArchive.Buffer.Empty();
+		RememberedDelegate.ExecuteIfBound( false, StreamArchive.IsSaving() );
 	}
 
-	RememberedDelegate.ExecuteIfBound( bSucceeded, Archive.IsSaving() );
+	// Reset delegate
 	RememberedDelegate = FOnStreamReadyDelegate();
 }
 
+void FHttpNetworkReplayStreamer::Tick( float DeltaTime )
+{
+	if ( !SessionName.IsEmpty() )
+	{
+		const double FLUSH_TIME_IN_SECONDS = 10;
+
+		if ( FPlatformTime::Seconds() - LastFlushTime > FLUSH_TIME_IN_SECONDS )
+		{
+			FlushStream();
+		}
+	}
+}
+
+bool FHttpNetworkReplayStreamer::IsTickable() const
+{
+	return true;
+}
+
+TStatId FHttpNetworkReplayStreamer::GetStatId() const
+{
+	RETURN_QUICK_DECLARE_CYCLE_STAT( FHttpNetworkReplayStreamer, STATGROUP_Tickables );
+}
+
+
 IMPLEMENT_MODULE( FHttpNetworkReplayStreamingFactory, HttpNetworkReplayStreaming )
 
 TSharedPtr< INetworkReplayStreamer > FHttpNetworkReplayStreamingFactory::CreateReplayStreamer()
diff --git a/Engine/Source/Runtime/NetworkReplayStreaming/HttpNetworkReplayStreaming/Private/HttpNetworkReplayStreaming.h b/Engine/Source/Runtime/NetworkReplayStreaming/HttpNetworkReplayStreaming/Private/HttpNetworkReplayStreaming.h
--- a/Engine/Source/Runtime/NetworkReplayStreaming/HttpNetworkReplayStreaming/Private/HttpNetworkReplayStreaming.h
+++ b/Engine/Source/Runtime/NetworkReplayStreaming/HttpNetworkReplayStreaming/Private/HttpNetworkReplayStreaming.h
@@ -3,6 +3,7 @@
 
 #include "NetworkReplayStreaming.h"
 #include "HTTP.h"
+#include "Runtime/Engine/Public/Tickable.h"
 
 /**
  * Archive used to buffer stream over http
@@ -24,21 +25,32 @@ public:
 /**
  * Http network replay streaming manager
  */
-class FHttpNetworkReplayStreamer : public INetworkReplayStreamer
+class FHttpNetworkReplayStreamer : public INetworkReplayStreamer, public FTickableGameObject
 {
 public:
-	/** IHttpNetworkReplayStreamer implementation */
-	FHttpNetworkReplayStreamer() {}
+	/** INetworkReplayStreamer implementation */
+	FHttpNetworkReplayStreamer() : StreamFileCount( 0 ), LastFlushTime( 0 ) {}
 	virtual void StartStreaming( FString& StreamName, bool bRecord, const FOnStreamReadyDelegate& Delegate ) override;
 	virtual void StopStreaming() override;
 	virtual FArchive* GetStreamingArchive() override;
 	virtual FArchive* GetMetadataArchive() override;
 
-	void HttpRequestFinished( FHttpRequestPtr HttpRequest, FHttpResponsePtr HttpResponse, bool bSucceeded );
+	/** FHttpNetworkReplayStreamer */
+	void FlushStream();
 
-	FOnStreamReadyDelegate	RememberedDelegate;
-	HttpStreamFArchive		Archive;
-	FString					DemoShortName;
+	void HttpDownloadFinished( FHttpRequestPtr HttpRequest, FHttpResponsePtr HttpResponse, bool bSucceeded );
+	void HttpStartStreamingUpFinished( FHttpRequestPtr HttpRequest, FHttpResponsePtr HttpResponse, bool bSucceeded );
+
+	/** FTickableGameObject */
+	virtual void Tick( float DeltaTime ) override;
+	virtual bool IsTickable() const override;
+	virtual TStatId GetStatId() const override;
+
+	FOnStreamReadyDelegate	RememberedDelegate;		// Delegate passed in to StartStreaming
+	HttpStreamFArchive		StreamArchive;			// Archive used to buffer the stream
+	FString					SessionName;			// Name of the session on the http replay server
+	int32					StreamFileCount;		// Used as a counter to increment the stream.x extension count
+	double					LastFlushTime;
 };
 
 class FHttpNetworkReplayStreamingFactory : public INetworkReplayStreamingFactory
