Http replay streaming factory now keeps track of created streamers so it can be in charge of allowing http requests to finalize before destroying + better tracking of when we're idle

diff --git a/Engine/Source/Runtime/NetworkReplayStreaming/HttpNetworkReplayStreaming/Private/HttpNetworkReplayStreaming.cpp b/Engine/Source/Runtime/NetworkReplayStreaming/HttpNetworkReplayStreaming/Private/HttpNetworkReplayStreaming.cpp
--- a/Engine/Source/Runtime/NetworkReplayStreaming/HttpNetworkReplayStreaming/Private/HttpNetworkReplayStreaming.cpp
+++ b/Engine/Source/Runtime/NetworkReplayStreaming/HttpNetworkReplayStreaming/Private/HttpNetworkReplayStreaming.cpp
@@ -83,6 +83,8 @@ void FHttpNetworkReplayStreamer::StartStreaming( FString& StreamName, bool bReco
 
 		HttpRequest->OnProcessRequestComplete().BindRaw( this, &FHttpNetworkReplayStreamer::HttpDownloadFinished );
 
+		StreamState = EStreamState::DownloadingStream;
+
 		StreamFileCount++;
 	}
 	else
@@ -102,6 +104,8 @@ void FHttpNetworkReplayStreamer::StartStreaming( FString& StreamName, bool bReco
 		HttpRequest->OnProcessRequestComplete().BindRaw( this, &FHttpNetworkReplayStreamer::HttpStartStreamingUpFinished );
 
 		SessionName.Empty();
+
+		StreamState = EStreamState::StartStreamingUp;
 	}
 	
 	HttpRequest->ProcessRequest();
@@ -127,9 +131,9 @@ void FHttpNetworkReplayStreamer::FlushStream()
 		// Create the Http request and add to pending request list
 		TSharedRef<class IHttpRequest> HttpRequest = FHttpModule::Get().CreateRequest();
 
-		// Currently, the engine free's this streamer as soon as it calls this, so we can't register for a callback here right now
-		// This is ok anyhow, since uploading is fire and forget (we aren't actually streaming it yet)
-		//HttpRequest->OnProcessRequestComplete().BindRaw( this, &FHttpNetworkReplayStreamer::HttpRequestFinished );
+		HttpRequest->OnProcessRequestComplete().BindRaw( this, &FHttpNetworkReplayStreamer::HttpUploadFinished );
+
+		StreamState = EStreamState::UploadingStream;
 
 		FString ServerURL;
 		GConfig->GetString( TEXT( "HttpNetworkReplayStreaming" ), TEXT( "ServerURL" ), ServerURL, GEngineIni );
@@ -138,6 +142,7 @@ void FHttpNetworkReplayStreamer::FlushStream()
 		HttpRequest->SetVerb( TEXT( "POST" ) );
 		HttpRequest->SetHeader( TEXT( "Content-Type" ), TEXT( "application/octet-stream" ) );
 		HttpRequest->SetContent( StreamArchive.Buffer );
+
 		HttpRequest->ProcessRequest();
 
 		StreamFileCount++;
@@ -161,6 +166,8 @@ FArchive* FHttpNetworkReplayStreamer::GetMetadataArchive()
 
 void FHttpNetworkReplayStreamer::HttpStartStreamingUpFinished( FHttpRequestPtr HttpRequest, FHttpResponsePtr HttpResponse, bool bSucceeded )
 {
+	check( StreamState == EStreamState::StartStreamingUp );
+
 	if ( bSucceeded && HttpResponse->GetResponseCode() == EHttpResponseCodes::Ok )
 	{
 		SessionName = HttpResponse->GetHeader( TEXT( "Session" ) );
@@ -169,12 +176,32 @@ void FHttpNetworkReplayStreamer::HttpStartStreamingUpFinished( FHttpRequestPtr H
 	}
 	else
 	{
-		UE_LOG( LogHttpReplay, Warning, TEXT( "FHttpNetworkReplayStreamer::HttpStartStreamingUpFinished. FAILED" ), *SessionName );
+		UE_LOG( LogHttpReplay, Warning, TEXT( "FHttpNetworkReplayStreamer::HttpStartStreamingUpFinished. FAILED" ) );
+	}
+
+	StreamState = EStreamState::Idle;
+}
+
+void FHttpNetworkReplayStreamer::HttpUploadFinished( FHttpRequestPtr HttpRequest, FHttpResponsePtr HttpResponse, bool bSucceeded )
+{
+	check( StreamState == EStreamState::UploadingStream );
+
+	if ( bSucceeded && HttpResponse->GetResponseCode() == EHttpResponseCodes::Ok )
+	{
+		UE_LOG( LogHttpReplay, Log, TEXT( "FHttpNetworkReplayStreamer::HttpUploadFinished." ) );
+	}
+	else
+	{
+		UE_LOG( LogHttpReplay, Warning, TEXT( "FHttpNetworkReplayStreamer::HttpUploadFinished. FAILED" ) );
 	}
+
+	StreamState = EStreamState::Idle;
 }
 
 void FHttpNetworkReplayStreamer::HttpDownloadFinished( FHttpRequestPtr HttpRequest, FHttpResponsePtr HttpResponse, bool bSucceeded )
 {
+	check( StreamState == EStreamState::DownloadingStream );
+
 	if ( bSucceeded && HttpResponse->GetResponseCode() == EHttpResponseCodes::Ok )
 	{
 		if ( StreamArchive.IsLoading() )
@@ -193,11 +220,13 @@ void FHttpNetworkReplayStreamer::HttpDownloadFinished( FHttpRequestPtr HttpReque
 
 	// Reset delegate
 	RememberedDelegate = FOnStreamReadyDelegate();
+
+	StreamState = EStreamState::Idle;
 }
 
 void FHttpNetworkReplayStreamer::Tick( float DeltaTime )
 {
-	if ( !SessionName.IsEmpty() )
+	if ( !SessionName.IsEmpty() && !IsBusy() )
 	{
 		const double FLUSH_TIME_IN_SECONDS = 10;
 
@@ -208,20 +237,42 @@ void FHttpNetworkReplayStreamer::Tick( float DeltaTime )
 	}
 }
 
-bool FHttpNetworkReplayStreamer::IsTickable() const
+bool FHttpNetworkReplayStreamer::IsBusy()
 {
-	return true;
+	return StreamState != EStreamState::Idle;
 }
 
-TStatId FHttpNetworkReplayStreamer::GetStatId() const
+IMPLEMENT_MODULE( FHttpNetworkReplayStreamingFactory, HttpNetworkReplayStreaming )
+
+TSharedPtr< INetworkReplayStreamer > FHttpNetworkReplayStreamingFactory::CreateReplayStreamer()
 {
-	RETURN_QUICK_DECLARE_CYCLE_STAT( FHttpNetworkReplayStreamer, STATGROUP_Tickables );
+	TSharedPtr< FHttpNetworkReplayStreamer > Streamer( new FHttpNetworkReplayStreamer );
+
+	HttpStreamers.Add( Streamer );
+
+	return Streamer;
 }
 
+void FHttpNetworkReplayStreamingFactory::Tick( float DeltaTime )
+{
+	for ( int i = HttpStreamers.Num() - 1; i >= 0; i-- )
+	{
+		HttpStreamers[i]->Tick( DeltaTime );
+		
+		// If we're the only holder of this streamer, and we're not busy, we can free it
+		if ( HttpStreamers[i].IsUnique() && !HttpStreamers[i]->IsBusy() )
+		{
+			HttpStreamers.RemoveAt( i );
+		}
+	}
+}
 
-IMPLEMENT_MODULE( FHttpNetworkReplayStreamingFactory, HttpNetworkReplayStreaming )
+bool FHttpNetworkReplayStreamingFactory::IsTickable() const
+{
+	return true;
+}
 
-TSharedPtr< INetworkReplayStreamer > FHttpNetworkReplayStreamingFactory::CreateReplayStreamer()
+TStatId FHttpNetworkReplayStreamingFactory::GetStatId() const
 {
-	return TSharedPtr< INetworkReplayStreamer >( new FHttpNetworkReplayStreamer );
+	RETURN_QUICK_DECLARE_CYCLE_STAT( FHttpNetworkReplayStreamingFactory, STATGROUP_Tickables );
 }
diff --git a/Engine/Source/Runtime/NetworkReplayStreaming/HttpNetworkReplayStreaming/Private/HttpNetworkReplayStreaming.h b/Engine/Source/Runtime/NetworkReplayStreaming/HttpNetworkReplayStreaming/Private/HttpNetworkReplayStreaming.h
--- a/Engine/Source/Runtime/NetworkReplayStreaming/HttpNetworkReplayStreaming/Private/HttpNetworkReplayStreaming.h
+++ b/Engine/Source/Runtime/NetworkReplayStreaming/HttpNetworkReplayStreaming/Private/HttpNetworkReplayStreaming.h
@@ -25,11 +25,11 @@ public:
 /**
  * Http network replay streaming manager
  */
-class FHttpNetworkReplayStreamer : public INetworkReplayStreamer, public FTickableGameObject
+class FHttpNetworkReplayStreamer : public INetworkReplayStreamer
 {
 public:
 	/** INetworkReplayStreamer implementation */
-	FHttpNetworkReplayStreamer() : StreamFileCount( 0 ), LastFlushTime( 0 ) {}
+	FHttpNetworkReplayStreamer() : StreamFileCount( 0 ), LastFlushTime( 0 ), StreamState( EStreamState::Idle ) {}
 	virtual void StartStreaming( FString& StreamName, bool bRecord, const FOnStreamReadyDelegate& Delegate ) override;
 	virtual void StopStreaming() override;
 	virtual FArchive* GetStreamingArchive() override;
@@ -40,21 +40,39 @@ public:
 
 	void HttpDownloadFinished( FHttpRequestPtr HttpRequest, FHttpResponsePtr HttpResponse, bool bSucceeded );
 	void HttpStartStreamingUpFinished( FHttpRequestPtr HttpRequest, FHttpResponsePtr HttpResponse, bool bSucceeded );
+	void HttpUploadFinished( FHttpRequestPtr HttpRequest, FHttpResponsePtr HttpResponse, bool bSucceeded );
 
-	/** FTickableGameObject */
-	virtual void Tick( float DeltaTime ) override;
-	virtual bool IsTickable() const override;
-	virtual TStatId GetStatId() const override;
+	void Tick( float DeltaTime );
+
+	bool IsBusy();
+
+	enum class EStreamState
+	{
+		Idle,
+		StartStreamingUp,
+		UploadingStream,
+		DownloadingStream,
+	};
 
 	FOnStreamReadyDelegate	RememberedDelegate;		// Delegate passed in to StartStreaming
 	HttpStreamFArchive		StreamArchive;			// Archive used to buffer the stream
 	FString					SessionName;			// Name of the session on the http replay server
 	int32					StreamFileCount;		// Used as a counter to increment the stream.x extension count
 	double					LastFlushTime;
+	EStreamState			StreamState;
+
 };
 
-class FHttpNetworkReplayStreamingFactory : public INetworkReplayStreamingFactory
+class FHttpNetworkReplayStreamingFactory : public INetworkReplayStreamingFactory, public FTickableGameObject
 {
 public:
+	/** INetworkReplayStreamingFactory */
 	virtual TSharedPtr< INetworkReplayStreamer > CreateReplayStreamer();
+
+	/** FTickableGameObject */
+	virtual void Tick( float DeltaTime ) override;
+	virtual bool IsTickable() const override;
+	virtual TStatId GetStatId() const override;
+
+	TArray< TSharedPtr< FHttpNetworkReplayStreamer > > HttpStreamers;
 };
