Engine side support for finalizing a network replay stream upload session, downloading, and cleanup

diff --git a/Engine/Source/Runtime/NetworkReplayStreaming/HttpNetworkReplayStreaming/Private/HttpNetworkReplayStreaming.cpp b/Engine/Source/Runtime/NetworkReplayStreaming/HttpNetworkReplayStreaming/Private/HttpNetworkReplayStreaming.cpp
--- a/Engine/Source/Runtime/NetworkReplayStreaming/HttpNetworkReplayStreaming/Private/HttpNetworkReplayStreaming.cpp
+++ b/Engine/Source/Runtime/NetworkReplayStreaming/HttpNetworkReplayStreaming/Private/HttpNetworkReplayStreaming.cpp
@@ -52,7 +52,17 @@ void HttpStreamFArchive::Seek( int64 InPos )
 
 void FHttpNetworkReplayStreamer::StartStreaming( FString& StreamName, bool bRecord, const FOnStreamReadyDelegate& Delegate )
 {
-	StopStreaming();
+	if ( !SessionName.IsEmpty() )
+	{
+		UE_LOG( LogHttpReplay, Warning, TEXT( "FHttpNetworkReplayStreamer::StartStreaming. SessionName already set." ) );
+		return;
+	}
+
+	if ( IsBusy() )
+	{
+		UE_LOG( LogHttpReplay, Warning, TEXT( "FHttpNetworkReplayStreamer::StartStreaming. BUSY" ) );
+		return;
+	}
 
 	RememberedDelegate = Delegate;
 
@@ -61,7 +71,7 @@ void FHttpNetworkReplayStreamer::StartStreaming( FString& StreamName, bool bReco
 
 	LastFlushTime = FPlatformTime::Seconds();
 
-	const bool bOverrideRecordingSession = false;
+	const bool bOverrideRecordingSession = true;//false;
 
 	// Override session name
 	if ( StreamArchive.ArIsLoading || bOverrideRecordingSession )
@@ -72,40 +82,39 @@ void FHttpNetworkReplayStreamer::StartStreaming( FString& StreamName, bool bReco
 	// Create the Http request and add to pending request list
 	TSharedRef<class IHttpRequest> HttpRequest = FHttpModule::Get().CreateRequest();
 
-	FString ServerURL;
 	GConfig->GetString( TEXT( "HttpNetworkReplayStreaming" ), TEXT( "ServerURL" ), ServerURL, GEngineIni );
 
 	if ( StreamArchive.ArIsLoading )
 	{
-		// Download the demo file from the http server
-		HttpRequest->SetURL( FString::Printf( TEXT( "%sdownload?Version=%s&Session=%s&Filename=stream.%i" ), *ServerURL, TEXT( "Test" ), *SessionName, StreamFileCount ) );
-		HttpRequest->SetVerb( TEXT( "GET" ) );
+		StreamFileCount = 0;
 
-		HttpRequest->OnProcessRequestComplete().BindRaw( this, &FHttpNetworkReplayStreamer::HttpDownloadFinished );
+		// Download the demo file from the http server
+		HttpRequest->SetURL( FString::Printf( TEXT( "%sstartdownloading?Version=%s&Session=%s" ), *ServerURL, TEXT( "Test" ), *SessionName ) );
+		HttpRequest->SetVerb( TEXT( "POST" ) );
 
-		StreamState = EStreamState::DownloadingStream;
+		HttpRequest->OnProcessRequestComplete().BindRaw( this, &FHttpNetworkReplayStreamer::HttpStartDownloadingFinished );
 
-		StreamFileCount++;
+		StreamState = EStreamState::StartDownloading;
 	}
 	else
 	{
 		// Make sure this session is registered with the http server
 		if ( !SessionName.IsEmpty() )
 		{
-			HttpRequest->SetURL( FString::Printf( TEXT( "%sstartstreamingup?Version=%s&Session=%s" ), *ServerURL, TEXT( "Test" ), *SessionName ) );
+			HttpRequest->SetURL( FString::Printf( TEXT( "%sstartuploading?Version=%s&Session=%s" ), *ServerURL, TEXT( "Test" ), *SessionName ) );
 		}
 		else
 		{
-			HttpRequest->SetURL( FString::Printf( TEXT( "%sstartstreamingup?Version=%s" ), *ServerURL, TEXT( "Test" ) ) );
+			HttpRequest->SetURL( FString::Printf( TEXT( "%sstartuploading?Version=%s" ), *ServerURL, TEXT( "Test" ) ) );
 		}
 
 		HttpRequest->SetVerb( TEXT( "POST" ) );
 
-		HttpRequest->OnProcessRequestComplete().BindRaw( this, &FHttpNetworkReplayStreamer::HttpStartStreamingUpFinished );
+		HttpRequest->OnProcessRequestComplete().BindRaw( this, &FHttpNetworkReplayStreamer::HttpStartUploadingFinished );
 
 		SessionName.Empty();
 
-		StreamState = EStreamState::StartStreamingUp;
+		StreamState = EStreamState::StartUploading;
 	}
 	
 	HttpRequest->ProcessRequest();
@@ -113,13 +122,7 @@ void FHttpNetworkReplayStreamer::StartStreaming( FString& StreamName, bool bReco
 
 void FHttpNetworkReplayStreamer::StopStreaming()
 {
-	FlushStream();
-
-	StreamArchive.ArIsLoading = false;
-	StreamArchive.ArIsSaving = false;
-	StreamArchive.Buffer.Empty();
-	StreamArchive.Pos = 0;
-	StreamFileCount = 0;
+	bFinalFlushNeeded = true;
 }
 
 void FHttpNetworkReplayStreamer::FlushStream()
@@ -135,9 +138,6 @@ void FHttpNetworkReplayStreamer::FlushStream()
 
 		StreamState = EStreamState::UploadingStream;
 
-		FString ServerURL;
-		GConfig->GetString( TEXT( "HttpNetworkReplayStreaming" ), TEXT( "ServerURL" ), ServerURL, GEngineIni );
-
 		HttpRequest->SetURL( FString::Printf( TEXT( "%supload?Version=%s&Session=%s&Filename=stream.%i" ), *ServerURL, TEXT( "Test" ), *SessionName, StreamFileCount ) );
 		HttpRequest->SetVerb( TEXT( "POST" ) );
 		HttpRequest->SetHeader( TEXT( "Content-Type" ), TEXT( "application/octet-stream" ) );
@@ -154,9 +154,23 @@ void FHttpNetworkReplayStreamer::FlushStream()
 	}
 }
 
+void FHttpNetworkReplayStreamer::DownloadNextChunk()
+{
+	TSharedRef<class IHttpRequest> HttpRequest = FHttpModule::Get().CreateRequest();
+
+	HttpRequest->SetURL( FString::Printf( TEXT( "%sdownload?Version=%s&Session=%s&Filename=stream.%i" ), *ServerURL, TEXT( "Test" ), *SessionName, StreamFileCount ) );
+	HttpRequest->SetVerb( TEXT( "GET" ) );
+
+	HttpRequest->OnProcessRequestComplete().BindRaw( this, &FHttpNetworkReplayStreamer::HttpDownloadFinished );
+
+	StreamState = EStreamState::DownloadingStream;
+
+	HttpRequest->ProcessRequest();
+}
+
 FArchive* FHttpNetworkReplayStreamer::GetStreamingArchive()
 {
-	return ( StreamArchive.IsSaving() || StreamArchive.Buffer.Num() >  0 ) ? &StreamArchive : NULL;
+	return ( StreamArchive.IsSaving() || ( StreamArchive.Buffer.Num() > 0 && StreamFileCount == NumDownloadChunks && NumDownloadChunks > 0 ) ) ? &StreamArchive : NULL;
 }
 
 FArchive* FHttpNetworkReplayStreamer::GetMetadataArchive()
@@ -164,64 +178,162 @@ FArchive* FHttpNetworkReplayStreamer::GetMetadataArchive()
 	return NULL;
 }
 
-void FHttpNetworkReplayStreamer::HttpStartStreamingUpFinished( FHttpRequestPtr HttpRequest, FHttpResponsePtr HttpResponse, bool bSucceeded )
+void FHttpNetworkReplayStreamer::HttpStartUploadingFinished( FHttpRequestPtr HttpRequest, FHttpResponsePtr HttpResponse, bool bSucceeded )
 {
-	check( StreamState == EStreamState::StartStreamingUp );
+	check( StreamState == EStreamState::StartUploading );
+
+	StreamState = EStreamState::Idle;
 
 	if ( bSucceeded && HttpResponse->GetResponseCode() == EHttpResponseCodes::Ok )
 	{
 		SessionName = HttpResponse->GetHeader( TEXT( "Session" ) );
 
-		UE_LOG( LogHttpReplay, Log, TEXT( "FHttpNetworkReplayStreamer::HttpStartStreamingUpFinished. SessionName: %s" ), *SessionName );
+		UE_LOG( LogHttpReplay, Log, TEXT( "FHttpNetworkReplayStreamer::HttpStartUploadingFinished. SessionName: %s" ), *SessionName );
 	}
 	else
 	{
-		UE_LOG( LogHttpReplay, Warning, TEXT( "FHttpNetworkReplayStreamer::HttpStartStreamingUpFinished. FAILED" ) );
+		UE_LOG( LogHttpReplay, Warning, TEXT( "FHttpNetworkReplayStreamer::HttpStartUploadingFinished. FAILED" ) );
 	}
+}
+
+void FHttpNetworkReplayStreamer::HttpStopUploadingFinished( FHttpRequestPtr HttpRequest, FHttpResponsePtr HttpResponse, bool bSucceeded )
+{
+	check( StreamState == EStreamState::StopUploading );
+	check( bFinalFlushNeeded );
 
 	StreamState = EStreamState::Idle;
+
+	if ( bSucceeded && HttpResponse->GetResponseCode() == EHttpResponseCodes::Ok )
+	{
+		UE_LOG( LogHttpReplay, Log, TEXT( "FHttpNetworkReplayStreamer::HttpStopUploadingFinished. SessionName: %s" ), *SessionName );
+	}
+	else
+	{
+		UE_LOG( LogHttpReplay, Warning, TEXT( "FHttpNetworkReplayStreamer::HttpStopUploadingFinished. FAILED" ) );
+	}
+
+	StreamArchive.ArIsLoading = false;
+	StreamArchive.ArIsSaving = false;
+	StreamArchive.Buffer.Empty();
+	StreamArchive.Pos = 0;
+	StreamFileCount = 0;
+	bFinalFlushNeeded = false;
+	SessionName.Empty();
 }
 
 void FHttpNetworkReplayStreamer::HttpUploadFinished( FHttpRequestPtr HttpRequest, FHttpResponsePtr HttpResponse, bool bSucceeded )
 {
 	check( StreamState == EStreamState::UploadingStream );
 
+	StreamState = EStreamState::Idle;
+
 	if ( bSucceeded && HttpResponse->GetResponseCode() == EHttpResponseCodes::Ok )
 	{
 		UE_LOG( LogHttpReplay, Log, TEXT( "FHttpNetworkReplayStreamer::HttpUploadFinished." ) );
+
+		if ( bFinalFlushNeeded )
+		{
+			// Create the Http request and add to pending request list
+			TSharedRef<class IHttpRequest> HttpRequest = FHttpModule::Get().CreateRequest();
+
+			HttpRequest->OnProcessRequestComplete().BindRaw( this, &FHttpNetworkReplayStreamer::HttpStopUploadingFinished );
+
+			StreamState = EStreamState::StopUploading;
+
+			HttpRequest->SetURL( FString::Printf( TEXT( "%sstopuploading?Version=%s&Session=%s&NumChunks=%i" ), *ServerURL, TEXT( "Test" ), *SessionName, StreamFileCount ) );
+			HttpRequest->SetVerb( TEXT( "POST" ) );
+			HttpRequest->SetHeader( TEXT( "Content-Type" ), TEXT( "application/octet-stream" ) );
+
+			HttpRequest->ProcessRequest();
+		}
 	}
 	else
 	{
 		UE_LOG( LogHttpReplay, Warning, TEXT( "FHttpNetworkReplayStreamer::HttpUploadFinished. FAILED" ) );
 	}
+}
+
+void FHttpNetworkReplayStreamer::HttpStartDownloadingFinished( FHttpRequestPtr HttpRequest, FHttpResponsePtr HttpResponse, bool bSucceeded )
+{
+	check( StreamState == EStreamState::StartDownloading );
 
 	StreamState = EStreamState::Idle;
+
+	if ( bSucceeded && HttpResponse->GetResponseCode() == EHttpResponseCodes::Ok )
+	{
+		FString NumChunksString = HttpResponse->GetHeader( TEXT( "NumChunks" ) );
+
+		NumDownloadChunks = FCString::Atoi( *NumChunksString );
+
+		UE_LOG( LogHttpReplay, Log, TEXT( "FHttpNetworkReplayStreamer::HttpStartDownloadingFinished. NumChunks: %i" ), NumDownloadChunks );
+
+		// Download the demo file from the http server
+		if ( NumDownloadChunks > 0 )
+		{
+			DownloadNextChunk();
+		}
+		else
+		{
+			UE_LOG( LogHttpReplay, Warning, TEXT( "FHttpNetworkReplayStreamer::HttpStartDownloadingFinished. NO CHUNKS" ) );
+
+			RememberedDelegate.ExecuteIfBound( false, StreamArchive.IsSaving() );
+
+			// Reset delegate
+			RememberedDelegate = FOnStreamReadyDelegate();
+		}
+	}
+	else
+	{
+		// FAIL
+		UE_LOG( LogHttpReplay, Warning, TEXT( "FHttpNetworkReplayStreamer::HttpStartDownloadingFinished. FAILED" ) );
+
+		RememberedDelegate.ExecuteIfBound( false, StreamArchive.IsSaving() );
+
+		// Reset delegate
+		RememberedDelegate = FOnStreamReadyDelegate();
+	}
 }
 
 void FHttpNetworkReplayStreamer::HttpDownloadFinished( FHttpRequestPtr HttpRequest, FHttpResponsePtr HttpResponse, bool bSucceeded )
 {
 	check( StreamState == EStreamState::DownloadingStream );
 
+	StreamState = EStreamState::Idle;
+
 	if ( bSucceeded && HttpResponse->GetResponseCode() == EHttpResponseCodes::Ok )
 	{
 		if ( StreamArchive.IsLoading() )
 		{
-			StreamArchive.Buffer = HttpResponse->GetContent();
+			StreamArchive.Buffer.Append( HttpResponse->GetContent() );
 		}
+		
+		StreamFileCount++;
+
+		UE_LOG( LogHttpReplay, Log, TEXT( "FHttpNetworkReplayStreamer::HttpDownloadFinished. %i / %i" ), StreamFileCount, NumDownloadChunks );
 
-		RememberedDelegate.ExecuteIfBound( true, StreamArchive.IsSaving() );
+		if ( StreamFileCount < NumDownloadChunks )
+		{
+			DownloadNextChunk();
+		}
+		else
+		{
+			RememberedDelegate.ExecuteIfBound( true, StreamArchive.IsSaving() );
+
+			// Reset delegate
+			RememberedDelegate = FOnStreamReadyDelegate();
+		}
 	}
 	else
 	{
 		// FAIL
+		UE_LOG( LogHttpReplay, Warning, TEXT( "FHttpNetworkReplayStreamer::HttpDownloadFinished. FAILED." ) );
+
 		StreamArchive.Buffer.Empty();
 		RememberedDelegate.ExecuteIfBound( false, StreamArchive.IsSaving() );
-	}
-
-	// Reset delegate
-	RememberedDelegate = FOnStreamReadyDelegate();
 
-	StreamState = EStreamState::Idle;
+		// Reset delegate
+		RememberedDelegate = FOnStreamReadyDelegate();
+	}
 }
 
 void FHttpNetworkReplayStreamer::Tick( float DeltaTime )
@@ -230,7 +342,7 @@ void FHttpNetworkReplayStreamer::Tick( float DeltaTime )
 	{
 		const double FLUSH_TIME_IN_SECONDS = 10;
 
-		if ( FPlatformTime::Seconds() - LastFlushTime > FLUSH_TIME_IN_SECONDS )
+		if ( bFinalFlushNeeded || FPlatformTime::Seconds() - LastFlushTime > FLUSH_TIME_IN_SECONDS )
 		{
 			FlushStream();
 		}
diff --git a/Engine/Source/Runtime/NetworkReplayStreaming/HttpNetworkReplayStreaming/Private/HttpNetworkReplayStreaming.h b/Engine/Source/Runtime/NetworkReplayStreaming/HttpNetworkReplayStreaming/Private/HttpNetworkReplayStreaming.h
--- a/Engine/Source/Runtime/NetworkReplayStreaming/HttpNetworkReplayStreaming/Private/HttpNetworkReplayStreaming.h
+++ b/Engine/Source/Runtime/NetworkReplayStreaming/HttpNetworkReplayStreaming/Private/HttpNetworkReplayStreaming.h
@@ -29,7 +29,7 @@ class FHttpNetworkReplayStreamer : public INetworkReplayStreamer
 {
 public:
 	/** INetworkReplayStreamer implementation */
-	FHttpNetworkReplayStreamer() : StreamFileCount( 0 ), LastFlushTime( 0 ), StreamState( EStreamState::Idle ) {}
+	FHttpNetworkReplayStreamer() : StreamFileCount( 0 ), LastFlushTime( 0 ), StreamState( EStreamState::Idle ), bFinalFlushNeeded( false ), NumDownloadChunks( 0 ) {}
 	virtual void StartStreaming( FString& StreamName, bool bRecord, const FOnStreamReadyDelegate& Delegate ) override;
 	virtual void StopStreaming() override;
 	virtual FArchive* GetStreamingArchive() override;
@@ -37,9 +37,12 @@ public:
 
 	/** FHttpNetworkReplayStreamer */
 	void FlushStream();
+	void DownloadNextChunk();
 
+	void HttpStartDownloadingFinished( FHttpRequestPtr HttpRequest, FHttpResponsePtr HttpResponse, bool bSucceeded );
 	void HttpDownloadFinished( FHttpRequestPtr HttpRequest, FHttpResponsePtr HttpResponse, bool bSucceeded );
-	void HttpStartStreamingUpFinished( FHttpRequestPtr HttpRequest, FHttpResponsePtr HttpResponse, bool bSucceeded );
+	void HttpStartUploadingFinished( FHttpRequestPtr HttpRequest, FHttpResponsePtr HttpResponse, bool bSucceeded );
+	void HttpStopUploadingFinished( FHttpRequestPtr HttpRequest, FHttpResponsePtr HttpResponse, bool bSucceeded );
 	void HttpUploadFinished( FHttpRequestPtr HttpRequest, FHttpResponsePtr HttpResponse, bool bSucceeded );
 
 	void Tick( float DeltaTime );
@@ -49,18 +52,22 @@ public:
 	enum class EStreamState
 	{
 		Idle,
-		StartStreamingUp,
+		StartUploading,
 		UploadingStream,
+		StopUploading,
+		StartDownloading,
 		DownloadingStream,
 	};
 
 	FOnStreamReadyDelegate	RememberedDelegate;		// Delegate passed in to StartStreaming
 	HttpStreamFArchive		StreamArchive;			// Archive used to buffer the stream
 	FString					SessionName;			// Name of the session on the http replay server
+	FString					ServerURL;
 	int32					StreamFileCount;		// Used as a counter to increment the stream.x extension count
 	double					LastFlushTime;
 	EStreamState			StreamState;
-
+	bool					bFinalFlushNeeded;
+	int32					NumDownloadChunks;
 };
 
 class FHttpNetworkReplayStreamingFactory : public INetworkReplayStreamingFactory, public FTickableGameObject
