Network replay changes.

* Added ability to check if data is available to replay streamer interface
* Added concept of header archive to replay streaming interface
* Http streamer now separates header into separate file
* Http streamer now supports downloading stream on demand in chunks
* Changes necessary to network replay engine to support downloading replay on demand

diff --git a/Engine/Source/Runtime/Engine/Private/DemoNetDriver.cpp b/Engine/Source/Runtime/Engine/Private/DemoNetDriver.cpp
--- a/Engine/Source/Runtime/Engine/Private/DemoNetDriver.cpp
+++ b/Engine/Source/Runtime/Engine/Private/DemoNetDriver.cpp
@@ -178,7 +178,7 @@ bool UDemoNetDriver::InitConnectInternal( FString& Error )
 
 	ResetDemoState();
 
-	FArchive* FileAr = ReplayStreamer->GetStreamingArchive();
+	FArchive* FileAr = ReplayStreamer->GetHeaderArchive();
 
 	if ( !FileAr )
 	{
@@ -188,10 +188,6 @@ bool UDemoNetDriver::InitConnectInternal( FString& Error )
 		return false;
 	}
 
-	// use the same byte format regardless of platform so that the demos are cross platform
-	// DEMO_FIXME: This is messing up for some reason, investigate
-	//FileAr->SetByteSwapping( true );
-
 	FNetworkDemoHeader DemoHeader;
 
 	(*FileAr) << DemoHeader;
@@ -345,7 +341,7 @@ bool UDemoNetDriver::InitListen( FNetworkNotify* InNotify, FURL& ListenURL, bool
 
 	ReplayStreamer->StartStreaming( DemoFilename, true, FOnStreamReadyDelegate::CreateUObject( this, &UDemoNetDriver::ReplayStreamingReady ) );
 
-	FArchive* FileAr = ReplayStreamer->GetStreamingArchive();
+	FArchive* FileAr = ReplayStreamer->GetHeaderArchive();
 
 	if( !FileAr )
 	{
@@ -675,6 +671,11 @@ void UDemoNetDriver::TickDemoRecord( float DeltaSeconds )
 
 bool UDemoNetDriver::ReadDemoFrame()
 {
+	if ( !ReplayStreamer->IsDataAvailable() )
+	{
+		return false;
+	}
+
 	FArchive* FileAr = ReplayStreamer->GetStreamingArchive();
 
 	if ( FileAr->IsError() )
diff --git a/Engine/Source/Runtime/NetworkReplayStreaming/HttpNetworkReplayStreaming/Private/HttpNetworkReplayStreaming.cpp b/Engine/Source/Runtime/NetworkReplayStreaming/HttpNetworkReplayStreaming/Private/HttpNetworkReplayStreaming.cpp
--- a/Engine/Source/Runtime/NetworkReplayStreaming/HttpNetworkReplayStreaming/Private/HttpNetworkReplayStreaming.cpp
+++ b/Engine/Source/Runtime/NetworkReplayStreaming/HttpNetworkReplayStreaming/Private/HttpNetworkReplayStreaming.cpp
@@ -69,6 +69,9 @@ void FHttpNetworkReplayStreamer::StartStreaming( FString& StreamName, bool bReco
 	StreamArchive.ArIsLoading = !bRecord;
 	StreamArchive.ArIsSaving = !StreamArchive.ArIsLoading;
 
+	HeaderArchive.ArIsLoading = !bRecord;
+	HeaderArchive.ArIsSaving = !StreamArchive.ArIsLoading;
+
 	LastFlushTime = FPlatformTime::Seconds();
 
 	const bool bOverrideRecordingSession = true;//false;
@@ -127,7 +130,41 @@ void FHttpNetworkReplayStreamer::StopStreaming()
 
 void FHttpNetworkReplayStreamer::FlushStream()
 {
-	if ( StreamArchive.ArIsSaving && StreamArchive.Buffer.Num() > 0 && !SessionName.IsEmpty() )
+	if ( SessionName.IsEmpty() || !StreamArchive.ArIsSaving )
+	{
+		return;
+	}
+
+	if ( IsBusy() )
+	{
+		return;
+	}
+
+	if ( HeaderArchive.Buffer.Num() > 0 )
+	{
+		// First upload the header
+		UE_LOG( LogHttpReplay, Log, TEXT( "FHttpNetworkReplayStreamer::FlushStream. Header. StreamFileCount: %i, Size: %i" ), StreamFileCount, StreamArchive.Buffer.Num() );
+
+		// Create the Http request and add to pending request list
+		TSharedRef<class IHttpRequest> HttpRequest = FHttpModule::Get().CreateRequest();
+
+		HttpRequest->OnProcessRequestComplete().BindRaw( this, &FHttpNetworkReplayStreamer::HttpHeaderUploadFinished );
+
+		StreamState = EStreamState::UploadingStream;
+
+		HttpRequest->SetURL( FString::Printf( TEXT( "%supload?Version=%s&Session=%s&Filename=replay.header" ), *ServerURL, TEXT( "Test" ), *SessionName ) );
+		HttpRequest->SetVerb( TEXT( "POST" ) );
+		HttpRequest->SetHeader( TEXT( "Content-Type" ), TEXT( "application/octet-stream" ) );
+		HttpRequest->SetContent( HeaderArchive.Buffer );
+
+		HeaderArchive.Buffer.Empty();
+		HeaderArchive.Pos = 0;
+
+		HttpRequest->ProcessRequest();
+
+		LastFlushTime = FPlatformTime::Seconds();
+	}
+	else if ( StreamArchive.Buffer.Num() > 0 )
 	{
 		UE_LOG( LogHttpReplay, Log, TEXT( "FHttpNetworkReplayStreamer::FlushStream. StreamFileCount: %i, Size: %i" ), StreamFileCount, StreamArchive.Buffer.Num() );
 
@@ -143,12 +180,12 @@ void FHttpNetworkReplayStreamer::FlushStream()
 		HttpRequest->SetHeader( TEXT( "Content-Type" ), TEXT( "application/octet-stream" ) );
 		HttpRequest->SetContent( StreamArchive.Buffer );
 
-		HttpRequest->ProcessRequest();
+		StreamArchive.Buffer.Empty();
+		StreamArchive.Pos = 0;
 
 		StreamFileCount++;
 
-		StreamArchive.Buffer.Empty();
-		StreamArchive.Pos = 0;
+		HttpRequest->ProcessRequest();
 
 		LastFlushTime = FPlatformTime::Seconds();
 	}
@@ -158,19 +195,35 @@ void FHttpNetworkReplayStreamer::DownloadNextChunk()
 {
 	TSharedRef<class IHttpRequest> HttpRequest = FHttpModule::Get().CreateRequest();
 
-	HttpRequest->SetURL( FString::Printf( TEXT( "%sdownload?Version=%s&Session=%s&Filename=stream.%i" ), *ServerURL, TEXT( "Test" ), *SessionName, StreamFileCount ) );
-	HttpRequest->SetVerb( TEXT( "GET" ) );
+	if ( HeaderArchive.Buffer.Num() == 0 )
+	{
+		// Download header first
+		HttpRequest->SetURL( FString::Printf( TEXT( "%sdownload?Version=%s&Session=%s&Filename=replay.header" ), *ServerURL, TEXT( "Test" ), *SessionName ) );
+		HttpRequest->SetVerb( TEXT( "GET" ) );
+
+		HttpRequest->OnProcessRequestComplete().BindRaw( this, &FHttpNetworkReplayStreamer::HttpDownloadHeaderFinished );
+	}
+	else
+	{
+		HttpRequest->SetURL( FString::Printf( TEXT( "%sdownload?Version=%s&Session=%s&Filename=stream.%i" ), *ServerURL, TEXT( "Test" ), *SessionName, StreamFileCount ) );
+		HttpRequest->SetVerb( TEXT( "GET" ) );
 
-	HttpRequest->OnProcessRequestComplete().BindRaw( this, &FHttpNetworkReplayStreamer::HttpDownloadFinished );
+		HttpRequest->OnProcessRequestComplete().BindRaw( this, &FHttpNetworkReplayStreamer::HttpDownloadFinished );
+	}
 
 	StreamState = EStreamState::DownloadingStream;
 
 	HttpRequest->ProcessRequest();
 }
 
+FArchive* FHttpNetworkReplayStreamer::GetHeaderArchive()
+{
+	return ( StreamArchive.IsSaving() || HeaderArchive.Pos < HeaderArchive.Buffer.Num() ) ? &HeaderArchive : NULL;
+}
+
 FArchive* FHttpNetworkReplayStreamer::GetStreamingArchive()
 {
-	return ( StreamArchive.IsSaving() || ( StreamArchive.Buffer.Num() > 0 && StreamFileCount == NumDownloadChunks && NumDownloadChunks > 0 ) ) ? &StreamArchive : NULL;
+	return ( StreamArchive.IsSaving() || IsDataAvailable() ) ? &StreamArchive : NULL;
 }
 
 FArchive* FHttpNetworkReplayStreamer::GetMetadataArchive()
@@ -178,6 +231,17 @@ FArchive* FHttpNetworkReplayStreamer::GetMetadataArchive()
 	return NULL;
 }
 
+bool FHttpNetworkReplayStreamer::IsDataAvailable() const
+{
+	//if ( StreamArchive.IsLoading() && StreamArchive.Pos < StreamArchive.Buffer.Num() && StreamFileCount == NumDownloadChunks && NumDownloadChunks > 0 )
+	if ( StreamArchive.IsLoading() && StreamArchive.Pos < StreamArchive.Buffer.Num() && NumDownloadChunks > 0 )
+	{
+		return true;
+	}
+	
+	return false;
+}
+
 void FHttpNetworkReplayStreamer::HttpStartUploadingFinished( FHttpRequestPtr HttpRequest, FHttpResponsePtr HttpResponse, bool bSucceeded )
 {
 	check( StreamState == EStreamState::StartUploading );
@@ -221,6 +285,22 @@ void FHttpNetworkReplayStreamer::HttpStopUploadingFinished( FHttpRequestPtr Http
 	SessionName.Empty();
 }
 
+void FHttpNetworkReplayStreamer::HttpHeaderUploadFinished( FHttpRequestPtr HttpRequest, FHttpResponsePtr HttpResponse, bool bSucceeded )
+{
+	check( StreamState == EStreamState::UploadingStream );
+
+	StreamState = EStreamState::Idle;
+
+	if ( bSucceeded && HttpResponse->GetResponseCode() == EHttpResponseCodes::Ok )
+	{
+		UE_LOG( LogHttpReplay, Log, TEXT( "FHttpNetworkReplayStreamer::HttpHeaderUploadFinished." ) );
+	}
+	else
+	{
+		UE_LOG( LogHttpReplay, Warning, TEXT( "FHttpNetworkReplayStreamer::HttpHeaderUploadFinished. FAILED" ) );
+	}
+}
+
 void FHttpNetworkReplayStreamer::HttpUploadFinished( FHttpRequestPtr HttpRequest, FHttpResponsePtr HttpResponse, bool bSucceeded )
 {
 	check( StreamState == EStreamState::UploadingStream );
@@ -262,10 +342,11 @@ void FHttpNetworkReplayStreamer::HttpStartDownloadingFinished( FHttpRequestPtr H
 	if ( bSucceeded && HttpResponse->GetResponseCode() == EHttpResponseCodes::Ok )
 	{
 		FString NumChunksString = HttpResponse->GetHeader( TEXT( "NumChunks" ) );
+		FString State = HttpResponse->GetHeader( TEXT( "State" ) );
 
 		NumDownloadChunks = FCString::Atoi( *NumChunksString );
 
-		UE_LOG( LogHttpReplay, Log, TEXT( "FHttpNetworkReplayStreamer::HttpStartDownloadingFinished. NumChunks: %i" ), NumDownloadChunks );
+		UE_LOG( LogHttpReplay, Log, TEXT( "FHttpNetworkReplayStreamer::HttpStartDownloadingFinished. State: %s, NumChunks: %i" ), *State, NumDownloadChunks );
 
 		// Download the demo file from the http server
 		if ( NumDownloadChunks > 0 )
@@ -294,34 +375,20 @@ void FHttpNetworkReplayStreamer::HttpStartDownloadingFinished( FHttpRequestPtr H
 	}
 }
 
-void FHttpNetworkReplayStreamer::HttpDownloadFinished( FHttpRequestPtr HttpRequest, FHttpResponsePtr HttpResponse, bool bSucceeded )
+void FHttpNetworkReplayStreamer::HttpDownloadHeaderFinished( FHttpRequestPtr HttpRequest, FHttpResponsePtr HttpResponse, bool bSucceeded )
 {
 	check( StreamState == EStreamState::DownloadingStream );
+	check( StreamArchive.IsLoading() );
 
 	StreamState = EStreamState::Idle;
 
 	if ( bSucceeded && HttpResponse->GetResponseCode() == EHttpResponseCodes::Ok )
 	{
-		if ( StreamArchive.IsLoading() )
-		{
-			StreamArchive.Buffer.Append( HttpResponse->GetContent() );
-		}
-		
-		StreamFileCount++;
-
-		UE_LOG( LogHttpReplay, Log, TEXT( "FHttpNetworkReplayStreamer::HttpDownloadFinished. %i / %i" ), StreamFileCount, NumDownloadChunks );
+		HeaderArchive.Buffer.Append( HttpResponse->GetContent() );
 
-		if ( StreamFileCount < NumDownloadChunks )
-		{
-			DownloadNextChunk();
-		}
-		else
-		{
-			RememberedDelegate.ExecuteIfBound( true, StreamArchive.IsSaving() );
+		UE_LOG( LogHttpReplay, Log, TEXT( "FHttpNetworkReplayStreamer::HttpDownloadHeaderFinished. Size: %i" ), HeaderArchive.Buffer.Num() );
 
-			// Reset delegate
-			RememberedDelegate = FOnStreamReadyDelegate();
-		}
+		RememberedDelegate.ExecuteIfBound( true, StreamArchive.IsSaving() );
 	}
 	else
 	{
@@ -330,9 +397,33 @@ void FHttpNetworkReplayStreamer::HttpDownloadFinished( FHttpRequestPtr HttpReque
 
 		StreamArchive.Buffer.Empty();
 		RememberedDelegate.ExecuteIfBound( false, StreamArchive.IsSaving() );
+	}
 
-		// Reset delegate
-		RememberedDelegate = FOnStreamReadyDelegate();
+	// Reset delegate
+	RememberedDelegate = FOnStreamReadyDelegate();
+}
+
+void FHttpNetworkReplayStreamer::HttpDownloadFinished( FHttpRequestPtr HttpRequest, FHttpResponsePtr HttpResponse, bool bSucceeded )
+{
+	check( StreamState == EStreamState::DownloadingStream );
+	check( StreamArchive.IsLoading() );
+
+	StreamState = EStreamState::Idle;
+
+	if ( bSucceeded && HttpResponse->GetResponseCode() == EHttpResponseCodes::Ok )
+	{
+		StreamArchive.Buffer.Append( HttpResponse->GetContent() );
+		
+		StreamFileCount++;
+
+		UE_LOG( LogHttpReplay, Log, TEXT( "FHttpNetworkReplayStreamer::HttpDownloadFinished. Progress: %i / %i" ), StreamFileCount, NumDownloadChunks );
+	}
+	else
+	{
+		// FAIL
+		UE_LOG( LogHttpReplay, Warning, TEXT( "FHttpNetworkReplayStreamer::HttpDownloadFinished. FAILED." ) );
+
+		StreamArchive.Buffer.Empty();
 	}
 }
 
@@ -340,11 +431,21 @@ void FHttpNetworkReplayStreamer::Tick( float DeltaTime )
 {
 	if ( !SessionName.IsEmpty() && !IsBusy() )
 	{
-		const double FLUSH_TIME_IN_SECONDS = 10;
+		if ( StreamArchive.IsSaving() )
+		{
+			const double FLUSH_TIME_IN_SECONDS = 10;
 
-		if ( bFinalFlushNeeded || FPlatformTime::Seconds() - LastFlushTime > FLUSH_TIME_IN_SECONDS )
+			if ( bFinalFlushNeeded || HeaderArchive.Buffer.Num() > 0 || FPlatformTime::Seconds() - LastFlushTime > FLUSH_TIME_IN_SECONDS )
+			{
+				FlushStream();
+			}
+		}
+		else if ( StreamArchive.IsLoading() )
 		{
-			FlushStream();
+			if ( StreamFileCount < NumDownloadChunks && NumDownloadChunks > 0 && StreamArchive.Pos == StreamArchive.Buffer.Num() )
+			{
+				DownloadNextChunk();
+			}
 		}
 	}
 }
diff --git a/Engine/Source/Runtime/NetworkReplayStreaming/HttpNetworkReplayStreaming/Private/HttpNetworkReplayStreaming.h b/Engine/Source/Runtime/NetworkReplayStreaming/HttpNetworkReplayStreaming/Private/HttpNetworkReplayStreaming.h
--- a/Engine/Source/Runtime/NetworkReplayStreaming/HttpNetworkReplayStreaming/Private/HttpNetworkReplayStreaming.h
+++ b/Engine/Source/Runtime/NetworkReplayStreaming/HttpNetworkReplayStreaming/Private/HttpNetworkReplayStreaming.h
@@ -32,17 +32,21 @@ public:
 	FHttpNetworkReplayStreamer() : StreamFileCount( 0 ), LastFlushTime( 0 ), StreamState( EStreamState::Idle ), bFinalFlushNeeded( false ), NumDownloadChunks( 0 ) {}
 	virtual void StartStreaming( FString& StreamName, bool bRecord, const FOnStreamReadyDelegate& Delegate ) override;
 	virtual void StopStreaming() override;
+	virtual FArchive* GetHeaderArchive() override;
 	virtual FArchive* GetStreamingArchive() override;
 	virtual FArchive* GetMetadataArchive() override;
+	virtual bool IsDataAvailable() const override;
 
 	/** FHttpNetworkReplayStreamer */
 	void FlushStream();
 	void DownloadNextChunk();
 
 	void HttpStartDownloadingFinished( FHttpRequestPtr HttpRequest, FHttpResponsePtr HttpResponse, bool bSucceeded );
+	void HttpDownloadHeaderFinished( FHttpRequestPtr HttpRequest, FHttpResponsePtr HttpResponse, bool bSucceeded );
 	void HttpDownloadFinished( FHttpRequestPtr HttpRequest, FHttpResponsePtr HttpResponse, bool bSucceeded );
 	void HttpStartUploadingFinished( FHttpRequestPtr HttpRequest, FHttpResponsePtr HttpResponse, bool bSucceeded );
 	void HttpStopUploadingFinished( FHttpRequestPtr HttpRequest, FHttpResponsePtr HttpResponse, bool bSucceeded );
+	void HttpHeaderUploadFinished( FHttpRequestPtr HttpRequest, FHttpResponsePtr HttpResponse, bool bSucceeded );
 	void HttpUploadFinished( FHttpRequestPtr HttpRequest, FHttpResponsePtr HttpResponse, bool bSucceeded );
 
 	void Tick( float DeltaTime );
@@ -60,7 +64,8 @@ public:
 	};
 
 	FOnStreamReadyDelegate	RememberedDelegate;		// Delegate passed in to StartStreaming
-	HttpStreamFArchive		StreamArchive;			// Archive used to buffer the stream
+	HttpStreamFArchive		HeaderArchive;			// Archive used to buffer the header stream
+	HttpStreamFArchive		StreamArchive;			// Archive used to buffer the data stream
 	FString					SessionName;			// Name of the session on the http replay server
 	FString					ServerURL;
 	int32					StreamFileCount;		// Used as a counter to increment the stream.x extension count
diff --git a/Engine/Source/Runtime/NetworkReplayStreaming/NetworkReplayStreaming/Public/NetworkReplayStreaming.h b/Engine/Source/Runtime/NetworkReplayStreaming/NetworkReplayStreaming/Public/NetworkReplayStreaming.h
--- a/Engine/Source/Runtime/NetworkReplayStreaming/NetworkReplayStreaming/Public/NetworkReplayStreaming.h
+++ b/Engine/Source/Runtime/NetworkReplayStreaming/NetworkReplayStreaming/Public/NetworkReplayStreaming.h
@@ -17,8 +17,10 @@ public:
 
 	virtual void StartStreaming( FString& StreamName, bool bRecord, const FOnStreamReadyDelegate& Delegate ) = 0;
 	virtual void StopStreaming() = 0;
+	virtual FArchive* GetHeaderArchive() = 0;
 	virtual FArchive* GetStreamingArchive() = 0;
 	virtual FArchive* GetMetadataArchive() = 0;
+	virtual bool IsDataAvailable() const = 0;
 };
 
 /** Replay streamer factory */
diff --git a/Engine/Source/Runtime/NetworkReplayStreaming/NullNetworkReplayStreaming/Private/NullNetworkReplayStreaming.cpp b/Engine/Source/Runtime/NetworkReplayStreaming/NullNetworkReplayStreaming/Private/NullNetworkReplayStreaming.cpp
--- a/Engine/Source/Runtime/NetworkReplayStreaming/NullNetworkReplayStreaming/Private/NullNetworkReplayStreaming.cpp
+++ b/Engine/Source/Runtime/NetworkReplayStreaming/NullNetworkReplayStreaming/Private/NullNetworkReplayStreaming.cpp
@@ -67,6 +67,11 @@ void FNullNetworkReplayStreamer::StopStreaming()
 	MetadataFileAr = NULL;
 }
 
+FArchive* FNullNetworkReplayStreamer::GetHeaderArchive()
+{
+	return FileAr;
+}
+
 FArchive* FNullNetworkReplayStreamer::GetStreamingArchive()
 {
 	return FileAr;
diff --git a/Engine/Source/Runtime/NetworkReplayStreaming/NullNetworkReplayStreaming/Public/NullNetworkReplayStreaming.h b/Engine/Source/Runtime/NetworkReplayStreaming/NullNetworkReplayStreaming/Public/NullNetworkReplayStreaming.h
--- a/Engine/Source/Runtime/NetworkReplayStreaming/NullNetworkReplayStreaming/Public/NullNetworkReplayStreaming.h
+++ b/Engine/Source/Runtime/NetworkReplayStreaming/NullNetworkReplayStreaming/Public/NullNetworkReplayStreaming.h
@@ -13,8 +13,10 @@ public:
 	FNullNetworkReplayStreamer() : FileAr( NULL ), MetadataFileAr( NULL ) {}
 	virtual void StartStreaming( FString& StreamName, bool bRecord, const FOnStreamReadyDelegate& Delegate ) override;
 	virtual void StopStreaming() override;
+	virtual FArchive* GetHeaderArchive() override;
 	virtual FArchive* GetStreamingArchive() override;
 	virtual FArchive* GetMetadataArchive() override;
+	virtual bool IsDataAvailable() const override { return true; }
 
 private:
 	/** Handle to the archive that will read/write network packets */
