Http streamer cleanup. Better state management and transitions to more easily catch errors.

diff --git a/Engine/Source/Runtime/NetworkReplayStreaming/HttpNetworkReplayStreaming/Private/HttpNetworkReplayStreaming.cpp b/Engine/Source/Runtime/NetworkReplayStreaming/HttpNetworkReplayStreaming/Private/HttpNetworkReplayStreaming.cpp
--- a/Engine/Source/Runtime/NetworkReplayStreaming/HttpNetworkReplayStreaming/Private/HttpNetworkReplayStreaming.cpp
+++ b/Engine/Source/Runtime/NetworkReplayStreaming/HttpNetworkReplayStreaming/Private/HttpNetworkReplayStreaming.cpp
@@ -58,14 +58,16 @@ void FHttpNetworkReplayStreamer::StartStreaming( FString& StreamName, bool bReco
 		return;
 	}
 
-	if ( IsBusy() )
+	if ( IsHttpBusy() )
 	{
 		UE_LOG( LogHttpReplay, Warning, TEXT( "FHttpNetworkReplayStreamer::StartStreaming. BUSY" ) );
 		return;
 	}
 
+	// Remember the delegate, which we'll call as soon as the header is available
 	RememberedDelegate = Delegate;
 
+	// Setup the archives
 	StreamArchive.ArIsLoading = !bRecord;
 	StreamArchive.ArIsSaving = !StreamArchive.ArIsLoading;
 
@@ -76,32 +78,36 @@ void FHttpNetworkReplayStreamer::StartStreaming( FString& StreamName, bool bReco
 
 	const bool bOverrideRecordingSession = true;//false;
 
-	// Override session name
+	// Override session name if requested
 	if ( StreamArchive.ArIsLoading || bOverrideRecordingSession )
 	{
 		SessionName = StreamName;
 	}
 
+	// Initialize the server URL
+	GConfig->GetString( TEXT( "HttpNetworkReplayStreaming" ), TEXT( "ServerURL" ), ServerURL, GEngineIni );
+
 	// Create the Http request and add to pending request list
 	TSharedRef<class IHttpRequest> HttpRequest = FHttpModule::Get().CreateRequest();
 
-	GConfig->GetString( TEXT( "HttpNetworkReplayStreaming" ), TEXT( "ServerURL" ), ServerURL, GEngineIni );
+	StreamFileCount = 0;
 
 	if ( StreamArchive.ArIsLoading )
 	{
-		StreamFileCount = 0;
-
-		// Download the demo file from the http server
+		// Notify the http server that we want to start downloading a replay
 		HttpRequest->SetURL( FString::Printf( TEXT( "%sstartdownloading?Version=%s&Session=%s" ), *ServerURL, TEXT( "Test" ), *SessionName ) );
 		HttpRequest->SetVerb( TEXT( "POST" ) );
 
 		HttpRequest->OnProcessRequestComplete().BindRaw( this, &FHttpNetworkReplayStreamer::HttpStartDownloadingFinished );
 
-		StreamState = EStreamState::StartDownloading;
+		HttpState = EHttptate::StartDownloading;
+		
+		// Set the next streamer state to download the header
+		StreamerState = EStreamerState::NeedToDownloadHeader;
 	}
 	else
 	{
-		// Make sure this session is registered with the http server
+		// Notify the http server that we want to start upload a replay
 		if ( !SessionName.IsEmpty() )
 		{
 			HttpRequest->SetURL( FString::Printf( TEXT( "%sstartuploading?Version=%s&Session=%s" ), *ServerURL, TEXT( "Test" ), *SessionName ) );
@@ -117,7 +123,10 @@ void FHttpNetworkReplayStreamer::StartStreaming( FString& StreamName, bool bReco
 
 		SessionName.Empty();
 
-		StreamState = EStreamState::StartUploading;
+		HttpState = EHttptate::StartUploading;
+
+		// Set the next streamer state to upload the header
+		StreamerState = EStreamerState::NeedToUploadHeader;
 	}
 	
 	HttpRequest->ProcessRequest();
@@ -125,100 +134,140 @@ void FHttpNetworkReplayStreamer::StartStreaming( FString& StreamName, bool bReco
 
 void FHttpNetworkReplayStreamer::StopStreaming()
 {
-	bFinalFlushNeeded = true;
+	check( bStopStreamingCalled == false );
+
+	if ( StreamerState == EStreamerState::Idle )
+	{
+		return;
+	}
+
+	// We need to queued this operation up, and finish it once all other operations are done
+	// (wait for header to upload, and any current stream to finish uploading, then we can flush what's left in StreamArchive)
+	bStopStreamingCalled = true;
 }
 
-void FHttpNetworkReplayStreamer::FlushStream()
+void FHttpNetworkReplayStreamer::UploadHeader()
 {
-	if ( SessionName.IsEmpty() || !StreamArchive.ArIsSaving )
+	if ( SessionName.IsEmpty() || !HeaderArchive.ArIsSaving )
 	{
+		// IF there is no active session, or we are not recording, we don't need to flush
 		return;
 	}
 
-	if ( IsBusy() )
+	if ( HeaderArchive.Buffer.Num() == 0 )
 	{
+		// We haven't serialized the header yet, don't do anything yet
 		return;
 	}
 
-	if ( HeaderArchive.Buffer.Num() > 0 )
+	if ( IsHttpBusy() )
 	{
-		// First upload the header
-		UE_LOG( LogHttpReplay, Log, TEXT( "FHttpNetworkReplayStreamer::FlushStream. Header. StreamFileCount: %i, Size: %i" ), StreamFileCount, StreamArchive.Buffer.Num() );
+		// If we're currently busy we can't flush right now
+		return;
+	}
+
+	// First upload the header
+	UE_LOG( LogHttpReplay, Log, TEXT( "FHttpNetworkReplayStreamer::UploadHeader. Header. StreamFileCount: %i, Size: %i" ), StreamFileCount, StreamArchive.Buffer.Num() );
 
-		// Create the Http request and add to pending request list
-		TSharedRef<class IHttpRequest> HttpRequest = FHttpModule::Get().CreateRequest();
+	// Create the Http request and add to pending request list
+	TSharedRef<class IHttpRequest> HttpRequest = FHttpModule::Get().CreateRequest();
 
-		HttpRequest->OnProcessRequestComplete().BindRaw( this, &FHttpNetworkReplayStreamer::HttpHeaderUploadFinished );
+	HttpRequest->OnProcessRequestComplete().BindRaw( this, &FHttpNetworkReplayStreamer::HttpHeaderUploadFinished );
 
-		StreamState = EStreamState::UploadingStream;
+	HttpState = EHttptate::UploadingHeader;
 
-		HttpRequest->SetURL( FString::Printf( TEXT( "%supload?Version=%s&Session=%s&Filename=replay.header" ), *ServerURL, TEXT( "Test" ), *SessionName ) );
-		HttpRequest->SetVerb( TEXT( "POST" ) );
-		HttpRequest->SetHeader( TEXT( "Content-Type" ), TEXT( "application/octet-stream" ) );
-		HttpRequest->SetContent( HeaderArchive.Buffer );
+	HttpRequest->SetURL( FString::Printf( TEXT( "%supload?Version=%s&Session=%s&Filename=replay.header" ), *ServerURL, TEXT( "Test" ), *SessionName ) );
+	HttpRequest->SetVerb( TEXT( "POST" ) );
+	HttpRequest->SetHeader( TEXT( "Content-Type" ), TEXT( "application/octet-stream" ) );
+	HttpRequest->SetContent( HeaderArchive.Buffer );
 
-		HeaderArchive.Buffer.Empty();
-		HeaderArchive.Pos = 0;
+	// We're done with the header archive
+	HeaderArchive.Buffer.Empty();
+	HeaderArchive.Pos = 0;
 
-		HttpRequest->ProcessRequest();
+	HttpRequest->ProcessRequest();
+
+	LastFlushTime = FPlatformTime::Seconds();
+}
 
-		LastFlushTime = FPlatformTime::Seconds();
+void FHttpNetworkReplayStreamer::FlushStream()
+{
+	if ( SessionName.IsEmpty() || !StreamArchive.ArIsSaving )
+	{
+		// IF there is no active session, or we are not recording, we don't need to flush
+		return;
 	}
-	else if ( StreamArchive.Buffer.Num() > 0 )
+
+	if ( IsHttpBusy() )
 	{
-		UE_LOG( LogHttpReplay, Log, TEXT( "FHttpNetworkReplayStreamer::FlushStream. StreamFileCount: %i, Size: %i" ), StreamFileCount, StreamArchive.Buffer.Num() );
+		// If we're currently busy we can't flush right now
+		return;
+	}
 
-		// Create the Http request and add to pending request list
-		TSharedRef<class IHttpRequest> HttpRequest = FHttpModule::Get().CreateRequest();
+	// Upload any new streamed data to the http server
+	UE_LOG( LogHttpReplay, Log, TEXT( "FHttpNetworkReplayStreamer::FlushStream. StreamFileCount: %i, Size: %i" ), StreamFileCount, StreamArchive.Buffer.Num() );
 
-		HttpRequest->OnProcessRequestComplete().BindRaw( this, &FHttpNetworkReplayStreamer::HttpUploadFinished );
+	// Create the Http request and add to pending request list
+	TSharedRef<class IHttpRequest> HttpRequest = FHttpModule::Get().CreateRequest();
 
-		StreamState = EStreamState::UploadingStream;
+	HttpRequest->OnProcessRequestComplete().BindRaw( this, &FHttpNetworkReplayStreamer::HttpUploadFinished );
 
-		HttpRequest->SetURL( FString::Printf( TEXT( "%supload?Version=%s&Session=%s&Filename=stream.%i" ), *ServerURL, TEXT( "Test" ), *SessionName, StreamFileCount ) );
-		HttpRequest->SetVerb( TEXT( "POST" ) );
-		HttpRequest->SetHeader( TEXT( "Content-Type" ), TEXT( "application/octet-stream" ) );
-		HttpRequest->SetContent( StreamArchive.Buffer );
+	HttpState = EHttptate::UploadingStream;
 
-		StreamArchive.Buffer.Empty();
-		StreamArchive.Pos = 0;
+	HttpRequest->SetURL( FString::Printf( TEXT( "%supload?Version=%s&Session=%s&Filename=stream.%i" ), *ServerURL, TEXT( "Test" ), *SessionName, StreamFileCount ) );
+	HttpRequest->SetVerb( TEXT( "POST" ) );
+	HttpRequest->SetHeader( TEXT( "Content-Type" ), TEXT( "application/octet-stream" ) );
+	HttpRequest->SetContent( StreamArchive.Buffer );
 
-		StreamFileCount++;
+	StreamArchive.Buffer.Empty();
+	StreamArchive.Pos = 0;
 
-		HttpRequest->ProcessRequest();
+	StreamFileCount++;
 
-		LastFlushTime = FPlatformTime::Seconds();
-	}
+	HttpRequest->ProcessRequest();
+
+	LastFlushTime = FPlatformTime::Seconds();
 }
 
-void FHttpNetworkReplayStreamer::DownloadNextChunk()
+void FHttpNetworkReplayStreamer::DownloadHeader()
 {
+	// Download header first
 	TSharedRef<class IHttpRequest> HttpRequest = FHttpModule::Get().CreateRequest();
 
-	if ( HeaderArchive.Buffer.Num() == 0 )
-	{
-		// Download header first
-		HttpRequest->SetURL( FString::Printf( TEXT( "%sdownload?Version=%s&Session=%s&Filename=replay.header" ), *ServerURL, TEXT( "Test" ), *SessionName ) );
-		HttpRequest->SetVerb( TEXT( "GET" ) );
+	HttpRequest->SetURL( FString::Printf( TEXT( "%sdownload?Version=%s&Session=%s&Filename=replay.header" ), *ServerURL, TEXT( "Test" ), *SessionName ) );
+	HttpRequest->SetVerb( TEXT( "GET" ) );
 
-		HttpRequest->OnProcessRequestComplete().BindRaw( this, &FHttpNetworkReplayStreamer::HttpDownloadHeaderFinished );
-	}
-	else
-	{
-		HttpRequest->SetURL( FString::Printf( TEXT( "%sdownload?Version=%s&Session=%s&Filename=stream.%i" ), *ServerURL, TEXT( "Test" ), *SessionName, StreamFileCount ) );
-		HttpRequest->SetVerb( TEXT( "GET" ) );
+	HttpRequest->OnProcessRequestComplete().BindRaw( this, &FHttpNetworkReplayStreamer::HttpDownloadHeaderFinished );
+
+	HttpState = EHttptate::DownloadingHeader;
+
+	HttpRequest->ProcessRequest();
+}
 
-		HttpRequest->OnProcessRequestComplete().BindRaw( this, &FHttpNetworkReplayStreamer::HttpDownloadFinished );
+void FHttpNetworkReplayStreamer::DownloadNextChunk()
+{
+	if ( StreamFileCount >= NumDownloadChunks || NumDownloadChunks == 0 || StreamArchive.Pos < StreamArchive.Buffer.Num() )
+	{
+		// We don't need to download the next chunk yet
+		return;
 	}
 
-	StreamState = EStreamState::DownloadingStream;
+	TSharedRef<class IHttpRequest> HttpRequest = FHttpModule::Get().CreateRequest();
+
+	// Download the next stream chunk
+	HttpRequest->SetURL( FString::Printf( TEXT( "%sdownload?Version=%s&Session=%s&Filename=stream.%i" ), *ServerURL, TEXT( "Test" ), *SessionName, StreamFileCount ) );
+	HttpRequest->SetVerb( TEXT( "GET" ) );
+
+	HttpRequest->OnProcessRequestComplete().BindRaw( this, &FHttpNetworkReplayStreamer::HttpDownloadFinished );
+
+	HttpState = EHttptate::DownloadingStream;
 
 	HttpRequest->ProcessRequest();
 }
 
 FArchive* FHttpNetworkReplayStreamer::GetHeaderArchive()
 {
-	return ( StreamArchive.IsSaving() || HeaderArchive.Pos < HeaderArchive.Buffer.Num() ) ? &HeaderArchive : NULL;
+	return ( HeaderArchive.IsSaving() || HeaderArchive.Pos < HeaderArchive.Buffer.Num() ) ? &HeaderArchive : NULL;
 }
 
 FArchive* FHttpNetworkReplayStreamer::GetStreamingArchive()
@@ -233,7 +282,7 @@ FArchive* FHttpNetworkReplayStreamer::GetMetadataArchive()
 
 bool FHttpNetworkReplayStreamer::IsDataAvailable() const
 {
-	//if ( StreamArchive.IsLoading() && StreamArchive.Pos < StreamArchive.Buffer.Num() && StreamFileCount == NumDownloadChunks && NumDownloadChunks > 0 )
+	// If we are loading, and we have more data
 	if ( StreamArchive.IsLoading() && StreamArchive.Pos < StreamArchive.Buffer.Num() && NumDownloadChunks > 0 )
 	{
 		return true;
@@ -244,9 +293,10 @@ bool FHttpNetworkReplayStreamer::IsDataAvailable() const
 
 void FHttpNetworkReplayStreamer::HttpStartUploadingFinished( FHttpRequestPtr HttpRequest, FHttpResponsePtr HttpResponse, bool bSucceeded )
 {
-	check( StreamState == EStreamState::StartUploading );
+	check( HttpState == EHttptate::StartUploading );
+	check( StreamerState == EStreamerState::NeedToUploadHeader );
 
-	StreamState = EStreamState::Idle;
+	HttpState = EHttptate::Idle;
 
 	if ( bSucceeded && HttpResponse->GetResponseCode() == EHttpResponseCodes::Ok )
 	{
@@ -262,10 +312,11 @@ void FHttpNetworkReplayStreamer::HttpStartUploadingFinished( FHttpRequestPtr Htt
 
 void FHttpNetworkReplayStreamer::HttpStopUploadingFinished( FHttpRequestPtr HttpRequest, FHttpResponsePtr HttpResponse, bool bSucceeded )
 {
-	check( StreamState == EStreamState::StopUploading );
-	check( bFinalFlushNeeded );
+	check( HttpState == EHttptate::StopUploading );
+	check( StreamerState == EStreamerState::StreamingUpFinal );
 
-	StreamState = EStreamState::Idle;
+	HttpState = EHttptate::Idle;
+	StreamerState = EStreamerState::Idle;
 
 	if ( bSucceeded && HttpResponse->GetResponseCode() == EHttpResponseCodes::Ok )
 	{
@@ -281,44 +332,52 @@ void FHttpNetworkReplayStreamer::HttpStopUploadingFinished( FHttpRequestPtr Http
 	StreamArchive.Buffer.Empty();
 	StreamArchive.Pos = 0;
 	StreamFileCount = 0;
-	bFinalFlushNeeded = false;
+	bStopStreamingCalled = false;
 	SessionName.Empty();
 }
 
 void FHttpNetworkReplayStreamer::HttpHeaderUploadFinished( FHttpRequestPtr HttpRequest, FHttpResponsePtr HttpResponse, bool bSucceeded )
 {
-	check( StreamState == EStreamState::UploadingStream );
+	check( HttpState == EHttptate::UploadingHeader );
+	check( StreamerState == EStreamerState::NeedToUploadHeader );
 
-	StreamState = EStreamState::Idle;
+	HttpState = EHttptate::Idle;
 
 	if ( bSucceeded && HttpResponse->GetResponseCode() == EHttpResponseCodes::Ok )
 	{
 		UE_LOG( LogHttpReplay, Log, TEXT( "FHttpNetworkReplayStreamer::HttpHeaderUploadFinished." ) );
+
+		StreamerState = EStreamerState::StreamingUp;
 	}
 	else
 	{
 		UE_LOG( LogHttpReplay, Warning, TEXT( "FHttpNetworkReplayStreamer::HttpHeaderUploadFinished. FAILED" ) );
+
+		StreamerState = EStreamerState::Idle;
+
+		// FIXME: Notify demo driver with delegate
 	}
 }
 
 void FHttpNetworkReplayStreamer::HttpUploadFinished( FHttpRequestPtr HttpRequest, FHttpResponsePtr HttpResponse, bool bSucceeded )
 {
-	check( StreamState == EStreamState::UploadingStream );
+	check( HttpState == EHttptate::UploadingStream );
+	check( StreamerState == EStreamerState::StreamingUp || StreamerState == EStreamerState::StreamingUpFinal );
 
-	StreamState = EStreamState::Idle;
+	HttpState = EHttptate::Idle;
 
 	if ( bSucceeded && HttpResponse->GetResponseCode() == EHttpResponseCodes::Ok )
 	{
 		UE_LOG( LogHttpReplay, Log, TEXT( "FHttpNetworkReplayStreamer::HttpUploadFinished." ) );
 
-		if ( bFinalFlushNeeded )
+		if ( StreamerState == EStreamerState::StreamingUpFinal )
 		{
 			// Create the Http request and add to pending request list
 			TSharedRef<class IHttpRequest> HttpRequest = FHttpModule::Get().CreateRequest();
 
 			HttpRequest->OnProcessRequestComplete().BindRaw( this, &FHttpNetworkReplayStreamer::HttpStopUploadingFinished );
 
-			StreamState = EStreamState::StopUploading;
+			HttpState = EHttptate::StopUploading;
 
 			HttpRequest->SetURL( FString::Printf( TEXT( "%sstopuploading?Version=%s&Session=%s&NumChunks=%i" ), *ServerURL, TEXT( "Test" ), *SessionName, StreamFileCount ) );
 			HttpRequest->SetVerb( TEXT( "POST" ) );
@@ -335,9 +394,10 @@ void FHttpNetworkReplayStreamer::HttpUploadFinished( FHttpRequestPtr HttpRequest
 
 void FHttpNetworkReplayStreamer::HttpStartDownloadingFinished( FHttpRequestPtr HttpRequest, FHttpResponsePtr HttpResponse, bool bSucceeded )
 {
-	check( StreamState == EStreamState::StartDownloading );
+	check( HttpState == EHttptate::StartDownloading );
+	check( StreamerState == EStreamerState::NeedToDownloadHeader );
 
-	StreamState = EStreamState::Idle;
+	HttpState = EHttptate::Idle;
 
 	if ( bSucceeded && HttpResponse->GetResponseCode() == EHttpResponseCodes::Ok )
 	{
@@ -348,10 +408,10 @@ void FHttpNetworkReplayStreamer::HttpStartDownloadingFinished( FHttpRequestPtr H
 
 		UE_LOG( LogHttpReplay, Log, TEXT( "FHttpNetworkReplayStreamer::HttpStartDownloadingFinished. State: %s, NumChunks: %i" ), *State, NumDownloadChunks );
 
-		// Download the demo file from the http server
+		// First, download the header
 		if ( NumDownloadChunks > 0 )
-		{
-			DownloadNextChunk();
+		{			
+			StreamerState = EStreamerState::NeedToDownloadHeader;
 		}
 		else
 		{
@@ -361,6 +421,8 @@ void FHttpNetworkReplayStreamer::HttpStartDownloadingFinished( FHttpRequestPtr H
 
 			// Reset delegate
 			RememberedDelegate = FOnStreamReadyDelegate();
+
+			StreamerState = EStreamerState::Idle;
 		}
 	}
 	else
@@ -372,15 +434,18 @@ void FHttpNetworkReplayStreamer::HttpStartDownloadingFinished( FHttpRequestPtr H
 
 		// Reset delegate
 		RememberedDelegate = FOnStreamReadyDelegate();
+
+		StreamerState = EStreamerState::Idle;
 	}
 }
 
 void FHttpNetworkReplayStreamer::HttpDownloadHeaderFinished( FHttpRequestPtr HttpRequest, FHttpResponsePtr HttpResponse, bool bSucceeded )
 {
-	check( StreamState == EStreamState::DownloadingStream );
+	check( HttpState == EHttptate::DownloadingHeader );
+	check( StreamerState == EStreamerState::NeedToDownloadHeader );
 	check( StreamArchive.IsLoading() );
 
-	StreamState = EStreamState::Idle;
+	HttpState = EHttptate::Idle;
 
 	if ( bSucceeded && HttpResponse->GetResponseCode() == EHttpResponseCodes::Ok )
 	{
@@ -389,6 +454,7 @@ void FHttpNetworkReplayStreamer::HttpDownloadHeaderFinished( FHttpRequestPtr Htt
 		UE_LOG( LogHttpReplay, Log, TEXT( "FHttpNetworkReplayStreamer::HttpDownloadHeaderFinished. Size: %i" ), HeaderArchive.Buffer.Num() );
 
 		RememberedDelegate.ExecuteIfBound( true, StreamArchive.IsSaving() );
+		StreamerState = EStreamerState::StreamingDown;
 	}
 	else
 	{
@@ -397,6 +463,8 @@ void FHttpNetworkReplayStreamer::HttpDownloadHeaderFinished( FHttpRequestPtr Htt
 
 		StreamArchive.Buffer.Empty();
 		RememberedDelegate.ExecuteIfBound( false, StreamArchive.IsSaving() );
+
+		StreamerState = EStreamerState::Idle;
 	}
 
 	// Reset delegate
@@ -405,10 +473,10 @@ void FHttpNetworkReplayStreamer::HttpDownloadHeaderFinished( FHttpRequestPtr Htt
 
 void FHttpNetworkReplayStreamer::HttpDownloadFinished( FHttpRequestPtr HttpRequest, FHttpResponsePtr HttpResponse, bool bSucceeded )
 {
-	check( StreamState == EStreamState::DownloadingStream );
+	check( HttpState == EHttptate::DownloadingStream );
 	check( StreamArchive.IsLoading() );
 
-	StreamState = EStreamState::Idle;
+	HttpState = EHttptate::Idle;
 
 	if ( bSucceeded && HttpResponse->GetResponseCode() == EHttpResponseCodes::Ok )
 	{
@@ -429,30 +497,69 @@ void FHttpNetworkReplayStreamer::HttpDownloadFinished( FHttpRequestPtr HttpReque
 
 void FHttpNetworkReplayStreamer::Tick( float DeltaTime )
 {
-	if ( !SessionName.IsEmpty() && !IsBusy() )
+	if ( IsHttpBusy() )
 	{
-		if ( StreamArchive.IsSaving() )
-		{
-			const double FLUSH_TIME_IN_SECONDS = 10;
+		return;
+	}
 
-			if ( bFinalFlushNeeded || HeaderArchive.Buffer.Num() > 0 || FPlatformTime::Seconds() - LastFlushTime > FLUSH_TIME_IN_SECONDS )
-			{
-				FlushStream();
-			}
+	if ( SessionName.IsEmpty() )
+	{
+		check( StreamerState == EStreamerState::Idle );
+		return;
+	}
+
+	if ( bStopStreamingCalled )
+	{
+		// If http isn't busy and we need to flush the last chunk, then go to that state now
+		if ( StreamerState == EStreamerState::StreamingUp )
+		{
+			StreamerState = EStreamerState::StreamingUpFinal;
 		}
-		else if ( StreamArchive.IsLoading() )
+		else if ( StreamerState == EStreamerState::StreamingDown )
+		{
+			StreamerState = EStreamerState::Idle;	// FIXME: Implement StopDownloading
+		}
+
+		bStopStreamingCalled = false;
+	}
+
+	if ( StreamerState == EStreamerState::NeedToUploadHeader )
+	{
+		// If we're waiting on the header, don't do anything until the header has data and is uploaded
+		UploadHeader();
+	}
+	else if ( StreamerState == EStreamerState::StreamingUp )
+	{
+		const double FLUSH_TIME_IN_SECONDS = 10;
+
+		if ( FPlatformTime::Seconds() - LastFlushTime > FLUSH_TIME_IN_SECONDS )
 		{
-			if ( StreamFileCount < NumDownloadChunks && NumDownloadChunks > 0 && StreamArchive.Pos == StreamArchive.Buffer.Num() )
-			{
-				DownloadNextChunk();
-			}
+			FlushStream();
 		}
 	}
+	else if ( StreamerState == EStreamerState::StreamingUpFinal )
+	{
+		FlushStream();
+	}
+	else if ( StreamerState == EStreamerState::NeedToDownloadHeader )
+	{
+		// If we're waiting on the header to download, don't do anything until the header has been downloaded
+		DownloadHeader();
+	}
+	else if ( StreamerState == EStreamerState::StreamingDown )
+	{
+		DownloadNextChunk();
+	}
+}
+
+bool FHttpNetworkReplayStreamer::IsHttpBusy() const
+{
+	return HttpState != EHttptate::Idle;
 }
 
-bool FHttpNetworkReplayStreamer::IsBusy()
+bool FHttpNetworkReplayStreamer::IsStreaming() const
 {
-	return StreamState != EStreamState::Idle;
+	return StreamerState != EStreamerState::Idle;
 }
 
 IMPLEMENT_MODULE( FHttpNetworkReplayStreamingFactory, HttpNetworkReplayStreaming )
@@ -472,8 +579,8 @@ void FHttpNetworkReplayStreamingFactory::Tick( float DeltaTime )
 	{
 		HttpStreamers[i]->Tick( DeltaTime );
 		
-		// If we're the only holder of this streamer, and we're not busy, we can free it
-		if ( HttpStreamers[i].IsUnique() && !HttpStreamers[i]->IsBusy() )
+		// We can release our hold when streaming is completely done
+		if ( HttpStreamers[i].IsUnique() && !HttpStreamers[i]->IsStreaming() )
 		{
 			HttpStreamers.RemoveAt( i );
 		}
diff --git a/Engine/Source/Runtime/NetworkReplayStreaming/HttpNetworkReplayStreaming/Private/HttpNetworkReplayStreaming.h b/Engine/Source/Runtime/NetworkReplayStreaming/HttpNetworkReplayStreaming/Private/HttpNetworkReplayStreaming.h
--- a/Engine/Source/Runtime/NetworkReplayStreaming/HttpNetworkReplayStreaming/Private/HttpNetworkReplayStreaming.h
+++ b/Engine/Source/Runtime/NetworkReplayStreaming/HttpNetworkReplayStreaming/Private/HttpNetworkReplayStreaming.h
@@ -29,7 +29,7 @@ class FHttpNetworkReplayStreamer : public INetworkReplayStreamer
 {
 public:
 	/** INetworkReplayStreamer implementation */
-	FHttpNetworkReplayStreamer() : StreamFileCount( 0 ), LastFlushTime( 0 ), StreamState( EStreamState::Idle ), bFinalFlushNeeded( false ), NumDownloadChunks( 0 ) {}
+	FHttpNetworkReplayStreamer() : StreamFileCount( 0 ), LastFlushTime( 0 ), HttpState( EHttptate::Idle ), StreamerState( EStreamerState::Idle ), bStopStreamingCalled( false ), NumDownloadChunks( 0 ) {}
 	virtual void StartStreaming( FString& StreamName, bool bRecord, const FOnStreamReadyDelegate& Delegate ) override;
 	virtual void StopStreaming() override;
 	virtual FArchive* GetHeaderArchive() override;
@@ -38,7 +38,9 @@ public:
 	virtual bool IsDataAvailable() const override;
 
 	/** FHttpNetworkReplayStreamer */
+	void UploadHeader();
 	void FlushStream();
+	void DownloadHeader();
 	void DownloadNextChunk();
 
 	void HttpStartDownloadingFinished( FHttpRequestPtr HttpRequest, FHttpResponsePtr HttpResponse, bool bSucceeded );
@@ -51,18 +53,31 @@ public:
 
 	void Tick( float DeltaTime );
 
-	bool IsBusy();
+	bool IsHttpBusy() const;
+	bool IsStreaming() const;
 
-	enum class EStreamState
+	enum class EHttptate
 	{
 		Idle,
 		StartUploading,
+		UploadingHeader,
 		UploadingStream,
 		StopUploading,
 		StartDownloading,
+		DownloadingHeader,
 		DownloadingStream,
 	};
 
+	enum class EStreamerState
+	{
+		Idle,
+		NeedToUploadHeader,
+		NeedToDownloadHeader,
+		StreamingUp,
+		StreamingDown,
+		StreamingUpFinal,
+	};
+
 	FOnStreamReadyDelegate	RememberedDelegate;		// Delegate passed in to StartStreaming
 	HttpStreamFArchive		HeaderArchive;			// Archive used to buffer the header stream
 	HttpStreamFArchive		StreamArchive;			// Archive used to buffer the data stream
@@ -70,8 +85,9 @@ public:
 	FString					ServerURL;
 	int32					StreamFileCount;		// Used as a counter to increment the stream.x extension count
 	double					LastFlushTime;
-	EStreamState			StreamState;
-	bool					bFinalFlushNeeded;
+	EStreamerState			StreamerState;
+	EHttptate				HttpState;
+	bool					bStopStreamingCalled;
 	int32					NumDownloadChunks;
 };
 
