Fix Linux CIS (case-sensitive filename, initialization order).

diff --git a/Engine/Source/Runtime/NetworkReplayStreaming/HttpNetworkReplayStreaming/Private/HttpNetworkReplayStreaming.h b/Engine/Source/Runtime/NetworkReplayStreaming/HttpNetworkReplayStreaming/Private/HttpNetworkReplayStreaming.h
--- a/Engine/Source/Runtime/NetworkReplayStreaming/HttpNetworkReplayStreaming/Private/HttpNetworkReplayStreaming.h
+++ b/Engine/Source/Runtime/NetworkReplayStreaming/HttpNetworkReplayStreaming/Private/HttpNetworkReplayStreaming.h
@@ -2,7 +2,7 @@
 #pragma once
 
 #include "NetworkReplayStreaming.h"
-#include "HTTP.h"
+#include "Http.h"
 #include "Runtime/Engine/Public/Tickable.h"
 
 /**
@@ -29,7 +29,7 @@ class FHttpNetworkReplayStreamer : public INetworkReplayStreamer
 {
 public:
 	/** INetworkReplayStreamer implementation */
-	FHttpNetworkReplayStreamer() : StreamFileCount( 0 ), LastFlushTime( 0 ), HttpState( EHttptate::Idle ), StreamerState( EStreamerState::Idle ), bStopStreamingCalled( false ), NumDownloadChunks( 0 ) {}
+	FHttpNetworkReplayStreamer() : StreamFileCount( 0 ), LastFlushTime( 0 ), StreamerState( EStreamerState::Idle ), HttpState( EHttptate::Idle ), bStopStreamingCalled( false ), NumDownloadChunks( 0 ) {}
 	virtual void StartStreaming( FString& StreamName, bool bRecord, const FOnStreamReadyDelegate& Delegate ) override;
 	virtual void StopStreaming() override;
 	virtual FArchive* GetHeaderArchive() override;
