Http replay streamer work

* Better state management and transitions to more easily catch errors
* You can now view matches that are live that then transition into final while you are viewing (archived)

diff --git a/Engine/Source/Runtime/NetworkReplayStreaming/HttpNetworkReplayStreaming/Private/HttpNetworkReplayStreaming.cpp b/Engine/Source/Runtime/NetworkReplayStreaming/HttpNetworkReplayStreaming/Private/HttpNetworkReplayStreaming.cpp
--- a/Engine/Source/Runtime/NetworkReplayStreaming/HttpNetworkReplayStreaming/Private/HttpNetworkReplayStreaming.cpp
+++ b/Engine/Source/Runtime/NetworkReplayStreaming/HttpNetworkReplayStreaming/Private/HttpNetworkReplayStreaming.cpp
@@ -58,9 +58,15 @@ void FHttpNetworkReplayStreamer::StartStreaming( FString& StreamName, bool bReco
 		return;
 	}
 
+	if ( IsStreaming() )
+	{
+		UE_LOG( LogHttpReplay, Warning, TEXT( "FHttpNetworkReplayStreamer::StartStreaming. IsStreaming == true." ) );
+		return;
+	}
+
 	if ( IsHttpBusy() )
 	{
-		UE_LOG( LogHttpReplay, Warning, TEXT( "FHttpNetworkReplayStreamer::StartStreaming. BUSY" ) );
+		UE_LOG( LogHttpReplay, Warning, TEXT( "FHttpNetworkReplayStreamer::StartStreaming. IsHttpBusy == true." ) );
 		return;
 	}
 
@@ -74,7 +80,7 @@ void FHttpNetworkReplayStreamer::StartStreaming( FString& StreamName, bool bReco
 	HeaderArchive.ArIsLoading = !bRecord;
 	HeaderArchive.ArIsSaving = !StreamArchive.ArIsLoading;
 
-	LastFlushTime = FPlatformTime::Seconds();
+	LastChunkTime = FPlatformTime::Seconds();
 
 	const bool bOverrideRecordingSession = true;//false;
 
@@ -166,6 +172,8 @@ void FHttpNetworkReplayStreamer::UploadHeader()
 		return;
 	}
 
+	check( StreamFileCount == 0 );
+
 	// First upload the header
 	UE_LOG( LogHttpReplay, Log, TEXT( "FHttpNetworkReplayStreamer::UploadHeader. Header. StreamFileCount: %i, Size: %i" ), StreamFileCount, StreamArchive.Buffer.Num() );
 
@@ -176,7 +184,7 @@ void FHttpNetworkReplayStreamer::UploadHeader()
 
 	HttpState = EHttptate::UploadingHeader;
 
-	HttpRequest->SetURL( FString::Printf( TEXT( "%supload?Version=%s&Session=%s&Filename=replay.header" ), *ServerURL, TEXT( "Test" ), *SessionName ) );
+	HttpRequest->SetURL( FString::Printf( TEXT( "%supload?Version=%s&Session=%s&NumChunks=%i&Filename=replay.header" ), *ServerURL, TEXT( "Test" ), *SessionName, StreamFileCount ) );
 	HttpRequest->SetVerb( TEXT( "POST" ) );
 	HttpRequest->SetHeader( TEXT( "Content-Type" ), TEXT( "application/octet-stream" ) );
 	HttpRequest->SetContent( HeaderArchive.Buffer );
@@ -187,7 +195,7 @@ void FHttpNetworkReplayStreamer::UploadHeader()
 
 	HttpRequest->ProcessRequest();
 
-	LastFlushTime = FPlatformTime::Seconds();
+	LastChunkTime = FPlatformTime::Seconds();
 }
 
 void FHttpNetworkReplayStreamer::FlushStream()
@@ -214,7 +222,7 @@ void FHttpNetworkReplayStreamer::FlushStream()
 
 	HttpState = EHttptate::UploadingStream;
 
-	HttpRequest->SetURL( FString::Printf( TEXT( "%supload?Version=%s&Session=%s&Filename=stream.%i" ), *ServerURL, TEXT( "Test" ), *SessionName, StreamFileCount ) );
+	HttpRequest->SetURL( FString::Printf( TEXT( "%supload?Version=%s&Session=%s&NumChunks=%i&Filename=stream.%i" ), *ServerURL, TEXT( "Test" ), *SessionName, StreamFileCount + 1, StreamFileCount ) );
 	HttpRequest->SetVerb( TEXT( "POST" ) );
 	HttpRequest->SetHeader( TEXT( "Content-Type" ), TEXT( "application/octet-stream" ) );
 	HttpRequest->SetContent( StreamArchive.Buffer );
@@ -226,7 +234,7 @@ void FHttpNetworkReplayStreamer::FlushStream()
 
 	HttpRequest->ProcessRequest();
 
-	LastFlushTime = FPlatformTime::Seconds();
+	LastChunkTime = FPlatformTime::Seconds();
 }
 
 void FHttpNetworkReplayStreamer::DownloadHeader()
@@ -246,9 +254,25 @@ void FHttpNetworkReplayStreamer::DownloadHeader()
 
 void FHttpNetworkReplayStreamer::DownloadNextChunk()
 {
-	if ( StreamFileCount >= NumDownloadChunks || NumDownloadChunks == 0 || StreamArchive.Pos < StreamArchive.Buffer.Num() )
+	const bool bMoreChunksDefinitelyAvilable = StreamFileCount < NumDownloadChunks;
+	const bool bMoreChunksMayBeAvilable = bStreamIsLive;
+
+	if ( !bMoreChunksDefinitelyAvilable && !bMoreChunksMayBeAvilable )
 	{
-		// We don't need to download the next chunk yet
+		StreamerState = EStreamerState::Idle;
+		UE_LOG( LogHttpReplay, Log, TEXT( "FHttpNetworkReplayStreamer::DownloadNextChunk. Done streaming." ) );
+		return;
+	}
+
+	// We don't need to download the next chunk yet
+	const double CHECK_FOR_NEXT_CHUNK_IN_SECONDS = 5;
+
+	const bool bTimeToDownloadChunk = ( FPlatformTime::Seconds() - LastChunkTime > CHECK_FOR_NEXT_CHUNK_IN_SECONDS );
+	const bool bReallyNeedToDownloadChunk = StreamArchive.Pos >= StreamArchive.Buffer.Num() && bMoreChunksDefinitelyAvilable;
+
+	if ( !bTimeToDownloadChunk && !bReallyNeedToDownloadChunk )
+	{
+		// Not time yet
 		return;
 	}
 
@@ -263,6 +287,8 @@ void FHttpNetworkReplayStreamer::DownloadNextChunk()
 	HttpState = EHttptate::DownloadingStream;
 
 	HttpRequest->ProcessRequest();
+
+	LastChunkTime = FPlatformTime::Seconds();
 }
 
 FArchive* FHttpNetworkReplayStreamer::GetHeaderArchive()
@@ -404,6 +430,8 @@ void FHttpNetworkReplayStreamer::HttpStartDownloadingFinished( FHttpRequestPtr H
 		FString NumChunksString = HttpResponse->GetHeader( TEXT( "NumChunks" ) );
 		FString State = HttpResponse->GetHeader( TEXT( "State" ) );
 
+		bStreamIsLive = State == TEXT( "Live" );
+
 		NumDownloadChunks = FCString::Atoi( *NumChunksString );
 
 		UE_LOG( LogHttpReplay, Log, TEXT( "FHttpNetworkReplayStreamer::HttpStartDownloadingFinished. State: %s, NumChunks: %i" ), *State, NumDownloadChunks );
@@ -480,18 +508,39 @@ void FHttpNetworkReplayStreamer::HttpDownloadFinished( FHttpRequestPtr HttpReque
 
 	if ( bSucceeded && HttpResponse->GetResponseCode() == EHttpResponseCodes::Ok )
 	{
-		StreamArchive.Buffer.Append( HttpResponse->GetContent() );
-		
-		StreamFileCount++;
+		FString NumChunksString = HttpResponse->GetHeader( TEXT( "NumChunks" ) );
+		FString State = HttpResponse->GetHeader( TEXT( "State" ) );
+
+		bStreamIsLive = State == TEXT( "Live" );
 
-		UE_LOG( LogHttpReplay, Log, TEXT( "FHttpNetworkReplayStreamer::HttpDownloadFinished. Progress: %i / %i" ), StreamFileCount, NumDownloadChunks );
+		NumDownloadChunks = FCString::Atoi( *NumChunksString );
+
+		if ( HttpResponse->GetContent().Num() > 0 || bStreamIsLive )
+		{
+			if ( HttpResponse->GetContent().Num() > 0 )
+			{
+				StreamArchive.Buffer.Append( HttpResponse->GetContent() );
+				StreamFileCount++;
+			}
+
+			UE_LOG( LogHttpReplay, Log, TEXT( "FHttpNetworkReplayStreamer::HttpDownloadFinished. State: %s, Progress: %i / %i" ), *State, StreamFileCount, NumDownloadChunks );
+		}
+		else
+		{
+			UE_LOG( LogHttpReplay, Warning, TEXT( "FHttpNetworkReplayStreamer::HttpDownloadFinished. FAILED." ) );
+			StreamArchive.Buffer.Empty();
+			StreamerState = EStreamerState::Idle;
+
+			// FIXME: Report error to engine
+		}		
 	}
 	else
 	{
-		// FAIL
 		UE_LOG( LogHttpReplay, Warning, TEXT( "FHttpNetworkReplayStreamer::HttpDownloadFinished. FAILED." ) );
-
 		StreamArchive.Buffer.Empty();
+		StreamerState = EStreamerState::Idle;
+
+		// FIXME: Report error to engine
 	}
 }
 
@@ -514,13 +563,13 @@ void FHttpNetworkReplayStreamer::Tick( float DeltaTime )
 		if ( StreamerState == EStreamerState::StreamingUp )
 		{
 			StreamerState = EStreamerState::StreamingUpFinal;
+			bStopStreamingCalled = false;
 		}
 		else if ( StreamerState == EStreamerState::StreamingDown )
 		{
 			StreamerState = EStreamerState::Idle;	// FIXME: Implement StopDownloading
+			bStopStreamingCalled = false;
 		}
-
-		bStopStreamingCalled = false;
 	}
 
 	if ( StreamerState == EStreamerState::NeedToUploadHeader )
@@ -532,7 +581,7 @@ void FHttpNetworkReplayStreamer::Tick( float DeltaTime )
 	{
 		const double FLUSH_TIME_IN_SECONDS = 10;
 
-		if ( FPlatformTime::Seconds() - LastFlushTime > FLUSH_TIME_IN_SECONDS )
+		if ( FPlatformTime::Seconds() - LastChunkTime > FLUSH_TIME_IN_SECONDS )
 		{
 			FlushStream();
 		}
@@ -580,8 +629,13 @@ void FHttpNetworkReplayStreamingFactory::Tick( float DeltaTime )
 		HttpStreamers[i]->Tick( DeltaTime );
 		
 		// We can release our hold when streaming is completely done
-		if ( HttpStreamers[i].IsUnique() && !HttpStreamers[i]->IsStreaming() )
+		if ( HttpStreamers[i].IsUnique() && !HttpStreamers[i]->IsHttpBusy() && !HttpStreamers[i]->bStopStreamingCalled )
 		{
+			if ( HttpStreamers[i]->IsStreaming() )
+			{
+				UE_LOG( LogHttpReplay, Warning, TEXT( "FHttpNetworkReplayStreamingFactory::Tick. Stream was stopped early." ) );
+			}
+
 			HttpStreamers.RemoveAt( i );
 		}
 	}
diff --git a/Engine/Source/Runtime/NetworkReplayStreaming/HttpNetworkReplayStreaming/Private/HttpNetworkReplayStreaming.h b/Engine/Source/Runtime/NetworkReplayStreaming/HttpNetworkReplayStreaming/Private/HttpNetworkReplayStreaming.h
--- a/Engine/Source/Runtime/NetworkReplayStreaming/HttpNetworkReplayStreaming/Private/HttpNetworkReplayStreaming.h
+++ b/Engine/Source/Runtime/NetworkReplayStreaming/HttpNetworkReplayStreaming/Private/HttpNetworkReplayStreaming.h
@@ -28,14 +28,23 @@ public:
 class FHttpNetworkReplayStreamer : public INetworkReplayStreamer
 {
 public:
+	FHttpNetworkReplayStreamer() : 
+		StreamFileCount( 0 ), 
+		LastChunkTime( 0 ), 
+		StreamerState( EStreamerState::Idle ), 
+		HttpState( EHttptate::Idle ), 
+		bStopStreamingCalled( false ), 
+		bStreamIsLive( false ),
+		NumDownloadChunks( 0 ) 
+	{}
+
 	/** INetworkReplayStreamer implementation */
-	FHttpNetworkReplayStreamer() : StreamFileCount( 0 ), LastFlushTime( 0 ), StreamerState( EStreamerState::Idle ), HttpState( EHttptate::Idle ), bStopStreamingCalled( false ), NumDownloadChunks( 0 ) {}
-	virtual void StartStreaming( FString& StreamName, bool bRecord, const FOnStreamReadyDelegate& Delegate ) override;
-	virtual void StopStreaming() override;
-	virtual FArchive* GetHeaderArchive() override;
-	virtual FArchive* GetStreamingArchive() override;
-	virtual FArchive* GetMetadataArchive() override;
-	virtual bool IsDataAvailable() const override;
+	virtual void		StartStreaming( FString& StreamName, bool bRecord, const FOnStreamReadyDelegate& Delegate ) override;
+	virtual void		StopStreaming() override;
+	virtual FArchive*	GetHeaderArchive() override;
+	virtual FArchive*	GetStreamingArchive() override;
+	virtual FArchive*	GetMetadataArchive() override;
+	virtual bool		IsDataAvailable() const override;
 
 	/** FHttpNetworkReplayStreamer */
 	void UploadHeader();
@@ -43,6 +52,7 @@ public:
 	void DownloadHeader();
 	void DownloadNextChunk();
 
+	/** Delegates */
 	void HttpStartDownloadingFinished( FHttpRequestPtr HttpRequest, FHttpResponsePtr HttpResponse, bool bSucceeded );
 	void HttpDownloadHeaderFinished( FHttpRequestPtr HttpRequest, FHttpResponsePtr HttpResponse, bool bSucceeded );
 	void HttpDownloadFinished( FHttpRequestPtr HttpRequest, FHttpResponsePtr HttpResponse, bool bSucceeded );
@@ -52,42 +62,44 @@ public:
 	void HttpUploadFinished( FHttpRequestPtr HttpRequest, FHttpResponsePtr HttpResponse, bool bSucceeded );
 
 	void Tick( float DeltaTime );
+	bool IsHttpBusy() const;		// True if we are waiting on http request response
+	bool IsStreaming() const;		// True if we are streaming a replay up or down
 
-	bool IsHttpBusy() const;
-	bool IsStreaming() const;
-
+	/** EHttptate - If not idle, there is a http request in flight */
 	enum class EHttptate
 	{
-		Idle,
-		StartUploading,
-		UploadingHeader,
-		UploadingStream,
-		StopUploading,
-		StartDownloading,
-		DownloadingHeader,
-		DownloadingStream,
+		Idle,						// There is no http request in flight
+		StartUploading,				// We have made a request to start uploading a replay
+		UploadingHeader,			// We are uploading the replay header
+		UploadingStream,			// We are in the process of uploading the replay stream
+		StopUploading,				// We have made the request to stop uploading a live replay stream
+		StartDownloading,			// We have made the request to start downloading a replay stream
+		DownloadingHeader,			// We are downloading the replay header
+		DownloadingStream,			// We are in the process of downloading the replay stream
 	};
 
+	/** EStreamerState - Overall state of the streamer */
 	enum class EStreamerState
 	{
-		Idle,
-		NeedToUploadHeader,
-		NeedToDownloadHeader,
-		StreamingUp,
-		StreamingDown,
-		StreamingUpFinal,
+		Idle,						// The streamer is idle. Either we haven't started streaming yet, or we are done
+		NeedToUploadHeader,			// We are waiting to upload the header
+		NeedToDownloadHeader,		// We are waiting to download the header
+		StreamingUp,				// We are in the process of streaming a replay to the http server
+		StreamingDown,				// We are in the process of streaming a replay from the http server
+		StreamingUpFinal,			// We are uploading the final stream
 	};
 
 	FOnStreamReadyDelegate	RememberedDelegate;		// Delegate passed in to StartStreaming
 	HttpStreamFArchive		HeaderArchive;			// Archive used to buffer the header stream
 	HttpStreamFArchive		StreamArchive;			// Archive used to buffer the data stream
 	FString					SessionName;			// Name of the session on the http replay server
-	FString					ServerURL;
+	FString					ServerURL;				// The address of the server
 	int32					StreamFileCount;		// Used as a counter to increment the stream.x extension count
-	double					LastFlushTime;
-	EStreamerState			StreamerState;
-	EHttptate				HttpState;
+	double					LastChunkTime;			// The last time we uploaded/downloaded a chunk
+	EStreamerState			StreamerState;			// Overall state of the streamer
+	EHttptate				HttpState;				// If not idle, there is a http request in flight
 	bool					bStopStreamingCalled;
+	bool					bStreamIsLive;			// If true, we are viewing a live stream
 	int32					NumDownloadChunks;
 };
 
