Http replay streamer work

* Actual engine network version is now used when uploading streams
* Fixed bug when replay server would shutdown and engine would continually keep trying to upload streams

diff --git a/Engine/Source/Runtime/Engine/Private/DemoNetDriver.cpp b/Engine/Source/Runtime/Engine/Private/DemoNetDriver.cpp
--- a/Engine/Source/Runtime/Engine/Private/DemoNetDriver.cpp
+++ b/Engine/Source/Runtime/Engine/Private/DemoNetDriver.cpp
@@ -16,6 +16,7 @@
 #include "GameFramework/SpectatorPawnMovement.h"
 #include "Engine/GameInstance.h"
 #include "NetworkReplayStreaming.h"
+#include "Net/UnrealNetwork.h"
 
 DEFINE_LOG_CATEGORY_STATIC( LogDemo, Log, All );
 
@@ -167,7 +168,9 @@ bool UDemoNetDriver::InitConnect( FNetworkNotify* InNotify, const FURL& ConnectU
 	ServerConnection = NewObject<UNetConnection>(GetTransientPackage(), UDemoNetConnection::StaticClass());
 	ServerConnection->InitConnection( this, USOCK_Pending, ConnectURL, 1000000 );
 
-	ReplayStreamer->StartStreaming( DemoFilename, false, FOnStreamReadyDelegate::CreateUObject( this, &UDemoNetDriver::ReplayStreamingReady ) );
+	FString VersionString = FString::Printf( TEXT( "%s_%u" ), FApp::GetGameName(), FNetworkVersion::GetLocalNetworkVersion() );
+
+	ReplayStreamer->StartStreaming( DemoFilename, false, VersionString, FOnStreamReadyDelegate::CreateUObject( this, &UDemoNetDriver::ReplayStreamingReady ) );
 
 	return true;
 }
@@ -339,7 +342,9 @@ bool UDemoNetDriver::InitListen( FNetworkNotify* InNotify, FURL& ListenURL, bool
 	Connection->InitSendBuffer();
 	ClientConnections.Add( Connection );
 
-	ReplayStreamer->StartStreaming( DemoFilename, true, FOnStreamReadyDelegate::CreateUObject( this, &UDemoNetDriver::ReplayStreamingReady ) );
+	FString VersionString = FString::Printf( TEXT( "%s_%u" ), FApp::GetGameName(), FNetworkVersion::GetLocalNetworkVersion() );
+
+	ReplayStreamer->StartStreaming( DemoFilename, true, VersionString, FOnStreamReadyDelegate::CreateUObject( this, &UDemoNetDriver::ReplayStreamingReady ) );
 
 	FArchive* FileAr = ReplayStreamer->GetHeaderArchive();
 
diff --git a/Engine/Source/Runtime/NetworkReplayStreaming/HttpNetworkReplayStreaming/Private/HttpNetworkReplayStreaming.cpp b/Engine/Source/Runtime/NetworkReplayStreaming/HttpNetworkReplayStreaming/Private/HttpNetworkReplayStreaming.cpp
--- a/Engine/Source/Runtime/NetworkReplayStreaming/HttpNetworkReplayStreaming/Private/HttpNetworkReplayStreaming.cpp
+++ b/Engine/Source/Runtime/NetworkReplayStreaming/HttpNetworkReplayStreaming/Private/HttpNetworkReplayStreaming.cpp
@@ -50,7 +50,7 @@ void HttpStreamFArchive::Seek( int64 InPos )
 	Pos = InPos;
 }
 
-void FHttpNetworkReplayStreamer::StartStreaming( FString& StreamName, bool bRecord, const FOnStreamReadyDelegate& Delegate )
+void FHttpNetworkReplayStreamer::StartStreaming( const FString& StreamName, bool bRecord, const FString& VersionString, const FOnStreamReadyDelegate& Delegate )
 {
 	if ( !SessionName.IsEmpty() )
 	{
@@ -70,6 +70,8 @@ void FHttpNetworkReplayStreamer::StartStreaming( FString& StreamName, bool bReco
 		return;
 	}
 
+	SessionVersion = VersionString;
+
 	// Remember the delegate, which we'll call as soon as the header is available
 	RememberedDelegate = Delegate;
 
@@ -82,7 +84,7 @@ void FHttpNetworkReplayStreamer::StartStreaming( FString& StreamName, bool bReco
 
 	LastChunkTime = FPlatformTime::Seconds();
 
-	const bool bOverrideRecordingSession = true;//false;
+	const bool bOverrideRecordingSession = false;
 
 	// Override session name if requested
 	if ( StreamArchive.ArIsLoading || bOverrideRecordingSession )
@@ -101,7 +103,7 @@ void FHttpNetworkReplayStreamer::StartStreaming( FString& StreamName, bool bReco
 	if ( StreamArchive.ArIsLoading )
 	{
 		// Notify the http server that we want to start downloading a replay
-		HttpRequest->SetURL( FString::Printf( TEXT( "%sstartdownloading?Version=%s&Session=%s" ), *ServerURL, TEXT( "Test" ), *SessionName ) );
+		HttpRequest->SetURL( FString::Printf( TEXT( "%sstartdownloading?Version=%s&Session=%s" ), *ServerURL, *SessionVersion, *SessionName ) );
 		HttpRequest->SetVerb( TEXT( "POST" ) );
 
 		HttpRequest->OnProcessRequestComplete().BindRaw( this, &FHttpNetworkReplayStreamer::HttpStartDownloadingFinished );
@@ -116,11 +118,11 @@ void FHttpNetworkReplayStreamer::StartStreaming( FString& StreamName, bool bReco
 		// Notify the http server that we want to start upload a replay
 		if ( !SessionName.IsEmpty() )
 		{
-			HttpRequest->SetURL( FString::Printf( TEXT( "%sstartuploading?Version=%s&Session=%s" ), *ServerURL, TEXT( "Test" ), *SessionName ) );
+			HttpRequest->SetURL( FString::Printf( TEXT( "%sstartuploading?Version=%s&Session=%s" ), *ServerURL, *SessionVersion, *SessionName ) );
 		}
 		else
 		{
-			HttpRequest->SetURL( FString::Printf( TEXT( "%sstartuploading?Version=%s" ), *ServerURL, TEXT( "Test" ) ) );
+			HttpRequest->SetURL( FString::Printf( TEXT( "%sstartuploading?Version=%s" ), *ServerURL, *SessionVersion ) );
 		}
 
 		HttpRequest->SetVerb( TEXT( "POST" ) );
@@ -184,7 +186,7 @@ void FHttpNetworkReplayStreamer::UploadHeader()
 
 	HttpState = EHttptate::UploadingHeader;
 
-	HttpRequest->SetURL( FString::Printf( TEXT( "%supload?Version=%s&Session=%s&NumChunks=%i&Filename=replay.header" ), *ServerURL, TEXT( "Test" ), *SessionName, StreamFileCount ) );
+	HttpRequest->SetURL( FString::Printf( TEXT( "%supload?Version=%s&Session=%s&NumChunks=%i&Filename=replay.header" ), *ServerURL, *SessionVersion, *SessionName, StreamFileCount ) );
 	HttpRequest->SetVerb( TEXT( "POST" ) );
 	HttpRequest->SetHeader( TEXT( "Content-Type" ), TEXT( "application/octet-stream" ) );
 	HttpRequest->SetContent( HeaderArchive.Buffer );
@@ -222,7 +224,7 @@ void FHttpNetworkReplayStreamer::FlushStream()
 
 	HttpState = EHttptate::UploadingStream;
 
-	HttpRequest->SetURL( FString::Printf( TEXT( "%supload?Version=%s&Session=%s&NumChunks=%i&Filename=stream.%i" ), *ServerURL, TEXT( "Test" ), *SessionName, StreamFileCount + 1, StreamFileCount ) );
+	HttpRequest->SetURL( FString::Printf( TEXT( "%supload?Version=%s&Session=%s&NumChunks=%i&Filename=stream.%i" ), *ServerURL, *SessionVersion, *SessionName, StreamFileCount + 1, StreamFileCount ) );
 	HttpRequest->SetVerb( TEXT( "POST" ) );
 	HttpRequest->SetHeader( TEXT( "Content-Type" ), TEXT( "application/octet-stream" ) );
 	HttpRequest->SetContent( StreamArchive.Buffer );
@@ -242,7 +244,7 @@ void FHttpNetworkReplayStreamer::DownloadHeader()
 	// Download header first
 	TSharedRef<class IHttpRequest> HttpRequest = FHttpModule::Get().CreateRequest();
 
-	HttpRequest->SetURL( FString::Printf( TEXT( "%sdownload?Version=%s&Session=%s&Filename=replay.header" ), *ServerURL, TEXT( "Test" ), *SessionName ) );
+	HttpRequest->SetURL( FString::Printf( TEXT( "%sdownload?Version=%s&Session=%s&Filename=replay.header" ), *ServerURL, *SessionVersion, *SessionName ) );
 	HttpRequest->SetVerb( TEXT( "GET" ) );
 
 	HttpRequest->OnProcessRequestComplete().BindRaw( this, &FHttpNetworkReplayStreamer::HttpDownloadHeaderFinished );
@@ -264,7 +266,7 @@ void FHttpNetworkReplayStreamer::DownloadNextChunk()
 		return;
 	}
 
-	// We don't need to download the next chunk yet
+	// Determine if it's time to download the next chunk
 	const double CHECK_FOR_NEXT_CHUNK_IN_SECONDS = 5;
 
 	const bool bTimeToDownloadChunk = ( FPlatformTime::Seconds() - LastChunkTime > CHECK_FOR_NEXT_CHUNK_IN_SECONDS );
@@ -279,7 +281,7 @@ void FHttpNetworkReplayStreamer::DownloadNextChunk()
 	TSharedRef<class IHttpRequest> HttpRequest = FHttpModule::Get().CreateRequest();
 
 	// Download the next stream chunk
-	HttpRequest->SetURL( FString::Printf( TEXT( "%sdownload?Version=%s&Session=%s&Filename=stream.%i" ), *ServerURL, TEXT( "Test" ), *SessionName, StreamFileCount ) );
+	HttpRequest->SetURL( FString::Printf( TEXT( "%sdownload?Version=%s&Session=%s&Filename=stream.%i" ), *ServerURL, *SessionVersion, *SessionName, StreamFileCount ) );
 	HttpRequest->SetVerb( TEXT( "GET" ) );
 
 	HttpRequest->OnProcessRequestComplete().BindRaw( this, &FHttpNetworkReplayStreamer::HttpDownloadFinished );
@@ -405,7 +407,7 @@ void FHttpNetworkReplayStreamer::HttpUploadFinished( FHttpRequestPtr HttpRequest
 
 			HttpState = EHttptate::StopUploading;
 
-			HttpRequest->SetURL( FString::Printf( TEXT( "%sstopuploading?Version=%s&Session=%s&NumChunks=%i" ), *ServerURL, TEXT( "Test" ), *SessionName, StreamFileCount ) );
+			HttpRequest->SetURL( FString::Printf( TEXT( "%sstopuploading?Version=%s&Session=%s&NumChunks=%i" ), *ServerURL, *SessionVersion, *SessionName, StreamFileCount ) );
 			HttpRequest->SetVerb( TEXT( "POST" ) );
 			HttpRequest->SetHeader( TEXT( "Content-Type" ), TEXT( "application/octet-stream" ) );
 
@@ -415,6 +417,7 @@ void FHttpNetworkReplayStreamer::HttpUploadFinished( FHttpRequestPtr HttpRequest
 	else
 	{
 		UE_LOG( LogHttpReplay, Warning, TEXT( "FHttpNetworkReplayStreamer::HttpUploadFinished. FAILED" ) );
+		StreamerState = EStreamerState::Idle;
 	}
 }
 
diff --git a/Engine/Source/Runtime/NetworkReplayStreaming/HttpNetworkReplayStreaming/Private/HttpNetworkReplayStreaming.h b/Engine/Source/Runtime/NetworkReplayStreaming/HttpNetworkReplayStreaming/Private/HttpNetworkReplayStreaming.h
--- a/Engine/Source/Runtime/NetworkReplayStreaming/HttpNetworkReplayStreaming/Private/HttpNetworkReplayStreaming.h
+++ b/Engine/Source/Runtime/NetworkReplayStreaming/HttpNetworkReplayStreaming/Private/HttpNetworkReplayStreaming.h
@@ -39,7 +39,7 @@ public:
 	{}
 
 	/** INetworkReplayStreamer implementation */
-	virtual void		StartStreaming( FString& StreamName, bool bRecord, const FOnStreamReadyDelegate& Delegate ) override;
+	virtual void		StartStreaming( const FString& StreamName, bool bRecord, const FString& VersionString, const FOnStreamReadyDelegate& Delegate ) override;
 	virtual void		StopStreaming() override;
 	virtual FArchive*	GetHeaderArchive() override;
 	virtual FArchive*	GetStreamingArchive() override;
@@ -93,6 +93,7 @@ public:
 	HttpStreamFArchive		HeaderArchive;			// Archive used to buffer the header stream
 	HttpStreamFArchive		StreamArchive;			// Archive used to buffer the data stream
 	FString					SessionName;			// Name of the session on the http replay server
+	FString					SessionVersion;			// Version of the session
 	FString					ServerURL;				// The address of the server
 	int32					StreamFileCount;		// Used as a counter to increment the stream.x extension count
 	double					LastChunkTime;			// The last time we uploaded/downloaded a chunk
diff --git a/Engine/Source/Runtime/NetworkReplayStreaming/NetworkReplayStreaming/Public/NetworkReplayStreaming.h b/Engine/Source/Runtime/NetworkReplayStreaming/NetworkReplayStreaming/Public/NetworkReplayStreaming.h
--- a/Engine/Source/Runtime/NetworkReplayStreaming/NetworkReplayStreaming/Public/NetworkReplayStreaming.h
+++ b/Engine/Source/Runtime/NetworkReplayStreaming/NetworkReplayStreaming/Public/NetworkReplayStreaming.h
@@ -15,7 +15,7 @@ class INetworkReplayStreamer
 public:
 	virtual ~INetworkReplayStreamer() {}
 
-	virtual void StartStreaming( FString& StreamName, bool bRecord, const FOnStreamReadyDelegate& Delegate ) = 0;
+	virtual void StartStreaming( const FString& StreamName, bool bRecord, const FString& VersionString, const FOnStreamReadyDelegate& Delegate ) = 0;
 	virtual void StopStreaming() = 0;
 	virtual FArchive* GetHeaderArchive() = 0;
 	virtual FArchive* GetStreamingArchive() = 0;
diff --git a/Engine/Source/Runtime/NetworkReplayStreaming/NullNetworkReplayStreaming/Private/NullNetworkReplayStreaming.cpp b/Engine/Source/Runtime/NetworkReplayStreaming/NullNetworkReplayStreaming/Private/NullNetworkReplayStreaming.cpp
--- a/Engine/Source/Runtime/NetworkReplayStreaming/NullNetworkReplayStreaming/Private/NullNetworkReplayStreaming.cpp
+++ b/Engine/Source/Runtime/NetworkReplayStreaming/NullNetworkReplayStreaming/Private/NullNetworkReplayStreaming.cpp
@@ -11,7 +11,7 @@
  * (which will require support for writing/reading at the same time)
  */
 
-void FNullNetworkReplayStreamer::StartStreaming( FString& StreamName, bool bRecord, const FOnStreamReadyDelegate& Delegate )
+void FNullNetworkReplayStreamer::StartStreaming( const FString& StreamName, bool bRecord, const FString& VersionString, const FOnStreamReadyDelegate& Delegate )
 {
 	int32 Year, Month, DayOfWeek, Day, Hour, Min, Sec, MSec;
 	FPlatformTime::SystemTime( Year, Month, DayOfWeek, Day, Hour, Min, Sec, MSec );
diff --git a/Engine/Source/Runtime/NetworkReplayStreaming/NullNetworkReplayStreaming/Public/NullNetworkReplayStreaming.h b/Engine/Source/Runtime/NetworkReplayStreaming/NullNetworkReplayStreaming/Public/NullNetworkReplayStreaming.h
--- a/Engine/Source/Runtime/NetworkReplayStreaming/NullNetworkReplayStreaming/Public/NullNetworkReplayStreaming.h
+++ b/Engine/Source/Runtime/NetworkReplayStreaming/NullNetworkReplayStreaming/Public/NullNetworkReplayStreaming.h
@@ -11,7 +11,7 @@ class FNullNetworkReplayStreamer : public INetworkReplayStreamer
 {
 public:
 	FNullNetworkReplayStreamer() : FileAr( NULL ), MetadataFileAr( NULL ) {}
-	virtual void StartStreaming( FString& StreamName, bool bRecord, const FOnStreamReadyDelegate& Delegate ) override;
+	virtual void StartStreaming( const FString& StreamName, bool bRecord, const FString& VersionString, const FOnStreamReadyDelegate& Delegate ) override;
 	virtual void StopStreaming() override;
 	virtual FArchive* GetHeaderArchive() override;
 	virtual FArchive* GetStreamingArchive() override;
