Added functions to NetworkReplayStreaming to enumerate and delete streams, and to check whether a stream is live.
Implemented in the Null module, stubbed for the HTTP module.
ShooterGame now uses these functions in the demo list menu.

diff --git a/Engine/Source/Runtime/NetworkReplayStreaming/HttpNetworkReplayStreaming/Private/HttpNetworkReplayStreaming.cpp b/Engine/Source/Runtime/NetworkReplayStreaming/HttpNetworkReplayStreaming/Private/HttpNetworkReplayStreaming.cpp
--- a/Engine/Source/Runtime/NetworkReplayStreaming/HttpNetworkReplayStreaming/Private/HttpNetworkReplayStreaming.cpp
+++ b/Engine/Source/Runtime/NetworkReplayStreaming/HttpNetworkReplayStreaming/Private/HttpNetworkReplayStreaming.cpp
@@ -319,6 +319,24 @@ bool FHttpNetworkReplayStreamer::IsDataAvailable() const
 	return false;
 }
 
+bool FHttpNetworkReplayStreamer::IsLive(const FString& StreamName) const 
+{
+	// Stubbed!
+	return false;
+}
+
+void FHttpNetworkReplayStreamer::DeleteFinishedStream( const FString& StreamName, const FOnDeleteFinishedStreamComplete& Delegate ) const
+{
+	// Stubbed!
+	Delegate.ExecuteIfBound(false);
+}
+
+void FHttpNetworkReplayStreamer::EnumerateStreams( const FOnEnumerateStreamsComplete& Delegate ) const
+{
+	// Stubbed!
+	Delegate.ExecuteIfBound(TArray<FNetworkReplayStreamInfo>());
+}
+
 void FHttpNetworkReplayStreamer::HttpStartUploadingFinished( FHttpRequestPtr HttpRequest, FHttpResponsePtr HttpResponse, bool bSucceeded )
 {
 	check( HttpState == EHttptate::StartUploading );
diff --git a/Engine/Source/Runtime/NetworkReplayStreaming/HttpNetworkReplayStreaming/Private/HttpNetworkReplayStreaming.h b/Engine/Source/Runtime/NetworkReplayStreaming/HttpNetworkReplayStreaming/Private/HttpNetworkReplayStreaming.h
--- a/Engine/Source/Runtime/NetworkReplayStreaming/HttpNetworkReplayStreaming/Private/HttpNetworkReplayStreaming.h
+++ b/Engine/Source/Runtime/NetworkReplayStreaming/HttpNetworkReplayStreaming/Private/HttpNetworkReplayStreaming.h
@@ -45,6 +45,9 @@ public:
 	virtual FArchive*	GetStreamingArchive() override;
 	virtual FArchive*	GetMetadataArchive() override;
 	virtual bool		IsDataAvailable() const override;
+	virtual bool		IsLive( const FString& StreamName ) const override;
+	virtual void		DeleteFinishedStream( const FString& StreamName, const FOnDeleteFinishedStreamComplete& Delegate ) const override;
+	virtual void		EnumerateStreams( const FOnEnumerateStreamsComplete& Delegate ) const override;
 
 	/** FHttpNetworkReplayStreamer */
 	void UploadHeader();
diff --git a/Engine/Source/Runtime/NetworkReplayStreaming/NetworkReplayStreaming/Public/NetworkReplayStreaming.h b/Engine/Source/Runtime/NetworkReplayStreaming/NetworkReplayStreaming/Public/NetworkReplayStreaming.h
--- a/Engine/Source/Runtime/NetworkReplayStreaming/NetworkReplayStreaming/Public/NetworkReplayStreaming.h
+++ b/Engine/Source/Runtime/NetworkReplayStreaming/NetworkReplayStreaming/Public/NetworkReplayStreaming.h
@@ -6,9 +6,39 @@
 #include "ModuleManager.h"
 #include "Core.h"
 
+/** Struct to store information about a stream, returned from search results. */
+struct FNetworkReplayStreamInfo
+{
+	/** The name of the stream */
+	FString Name;
+
+	/** The date and time the stream was recorded */
+	FDateTime Timestamp;
+
+	/** The size of the stream */
+	int64 SizeInBytes;
+
+	/** True if the stream is live and the game hasn't completed yet */
+	bool bIsLive;
+};
+
 DECLARE_MULTICAST_DELEGATE_TwoParams( FOnStreamReady, bool, bool );
 typedef FOnStreamReady::FDelegate FOnStreamReadyDelegate;
 
+/**
+ * Delegate called when DeleteFinishedStream() completes.
+ *
+ * @param bWasSuccessful Whether the stream was deleted.
+ */
+DECLARE_DELEGATE_OneParam( FOnDeleteFinishedStreamComplete, bool );
+
+/**
+ * Delegate called when EnumerateStreams() completes.
+ *
+ * @param Streams An array containing information about the streams that were found.
+ */
+DECLARE_DELEGATE_OneParam( FOnEnumerateStreamsComplete, const TArray<FNetworkReplayStreamInfo>& );
+
 /** Generic interface for network replay streaming */
 class INetworkReplayStreamer 
 {
@@ -21,6 +51,24 @@ public:
 	virtual FArchive* GetStreamingArchive() = 0;
 	virtual FArchive* GetMetadataArchive() = 0;
 	virtual bool IsDataAvailable() const = 0;
+	
+	/** Returns true if the given StreamName is a game currently in progress */
+	virtual bool IsLive( const FString& StreamName ) const = 0;
+
+	/**
+	 * Attempts to delete the stream with the specified name. May execute asynchronously.
+	 *
+	 * @param StreamName The name of the stream to delete
+	 * @param Delegate A delegate that will be executed if bound when the delete operation completes
+	 */
+	virtual void DeleteFinishedStream( const FString& StreamName, const FOnDeleteFinishedStreamComplete& Delegate) const = 0;
+
+	/**
+	 * Retrieves the streams that are available for viewing. May execute asynchronously.
+	 *
+	 * @param Delegate A delegate that will be executed if bound when the list of streams is available
+	 */
+	virtual void EnumerateStreams( const FOnEnumerateStreamsComplete& Delegate ) const = 0;
 };
 
 /** Replay streamer factory */
diff --git a/Engine/Source/Runtime/NetworkReplayStreaming/NullNetworkReplayStreaming/Private/NullNetworkReplayStreaming.cpp b/Engine/Source/Runtime/NetworkReplayStreaming/NullNetworkReplayStreaming/Private/NullNetworkReplayStreaming.cpp
--- a/Engine/Source/Runtime/NetworkReplayStreaming/NullNetworkReplayStreaming/Private/NullNetworkReplayStreaming.cpp
+++ b/Engine/Source/Runtime/NetworkReplayStreaming/NullNetworkReplayStreaming/Private/NullNetworkReplayStreaming.cpp
@@ -4,6 +4,8 @@
 #include "Paths.h"
 #include "EngineVersion.h"
 
+DEFINE_LOG_CATEGORY_STATIC( LogNullReplay, Log, All );
+
 /**
  * Very basic implementation of network replay streaming using the file system
  * As of now, there is just simple opening and closing of the stream, and handing off the stream for direct use
@@ -11,7 +13,7 @@
  * (which will require support for writing/reading at the same time)
  */
 
-void FNullNetworkReplayStreamer::StartStreaming( const FString& StreamName, bool bRecord, const FString& VersionString, const FOnStreamReadyDelegate& Delegate )
+static FString GetStreamBaseFilename(const FString& StreamName)
 {
 	int32 Year, Month, DayOfWeek, Day, Hour, Min, Sec, MSec;
 	FPlatformTime::SystemTime( Year, Month, DayOfWeek, Day, Hour, Min, Sec, MSec );
@@ -30,56 +32,176 @@ void FNullNetworkReplayStreamer::StartStreaming( const FString& StreamName, bool
 	DemoName.ReplaceInline( TEXT( " " ),	TEXT( "_" ) );
 	DemoName.ReplaceInline( TEXT( "%" ),	TEXT( "_" ) );
 
+	return DemoName;
+}
+
+static FString GetDemoPath()
+{
+	return FPaths::Combine(*FPaths::GameSavedDir(), TEXT( "Demos/" ));
+}
+
+static FString GetStreamDirectory(const FString& StreamName)
+{
+	const FString DemoName = GetStreamBaseFilename(StreamName);
+
+	// Directory for this demo
+	const FString DemoDir  = FPaths::Combine(*GetDemoPath(), *DemoName);
+	
+	return DemoDir;
+}
+
+static FString GetStreamFullBaseFilename(const FString& StreamName)
+{
+	return FPaths::Combine(*GetStreamDirectory(StreamName), *GetStreamBaseFilename(StreamName));
+}
+
+static FString GetDemoFilename(const FString& StreamName)
+{
+	return GetStreamFullBaseFilename(StreamName) + TEXT(".demo");
+}
+
+static FString GetMetadataFilename(const FString& StreamName)
+{
+	return GetStreamFullBaseFilename(StreamName) + TEXT(".metadata");
+}
+
+void FNullNetworkReplayStreamer::StartStreaming( const FString& StreamName, bool bRecord, const FString& VersionString, const FOnStreamReadyDelegate& Delegate )
+{
 	// Create a directory for this demo
-	const FString DemoDir = FPaths::Combine(*FPaths::GameSavedDir(), TEXT( "Demos" ), *DemoName);
+	const FString DemoDir = GetStreamDirectory(StreamName);
 
 	IFileManager::Get().MakeDirectory( *DemoDir, true );
 
-	// Demo filename without extension
-	const FString FullDemoBaseName = DemoDir + TEXT("/") + DemoName;
+	const FString FullDemoFilename = GetDemoFilename(StreamName);
+	const FString FullMetadataFilename = GetMetadataFilename(StreamName);
 
-	const FString FullDemoFilename = FullDemoBaseName + TEXT(".demo");
-	const FString FullMetadataFilename = FullDemoBaseName + TEXT(".metadata");
+	CurrentStreamName = StreamName;
 
 	if ( !bRecord )
 	{
 		// Open file for reading
-		FileAr = IFileManager::Get().CreateFileReader( *FullDemoFilename );
-		MetadataFileAr = IFileManager::Get().CreateFileReader( *FullMetadataFilename );
+		FileAr.Reset( IFileManager::Get().CreateFileReader( *FullDemoFilename ) );
+		StreamerState = EStreamerState::Playback;
 	}
 	else
 	{
 		// Open file for writing
-		FileAr = IFileManager::Get().CreateFileWriter( *FullDemoFilename );
-		MetadataFileAr = IFileManager::Get().CreateFileWriter( *FullMetadataFilename );
+		FileAr.Reset( IFileManager::Get().CreateFileWriter( *FullDemoFilename, FILEWRITE_AllowRead ) );
+		StreamerState = EStreamerState::Recording;
 	}
 
 	// Notify immediately
-	Delegate.ExecuteIfBound( FileAr != NULL, bRecord );
+	Delegate.ExecuteIfBound( FileAr.Get() != NULL, bRecord );
 }
 
 void FNullNetworkReplayStreamer::StopStreaming()
 {
-	delete FileAr;
-	FileAr = NULL;
+	FileAr.Reset();
+	MetadataFileAr.Reset();
 
-	delete MetadataFileAr;
-	MetadataFileAr = NULL;
+	CurrentStreamName.Empty();
+	StreamerState = EStreamerState::Idle;
 }
 
 FArchive* FNullNetworkReplayStreamer::GetHeaderArchive()
 {
-	return FileAr;
+	return FileAr.Get();
 }
 
 FArchive* FNullNetworkReplayStreamer::GetStreamingArchive()
 {
-	return FileAr;
+	return FileAr.Get();
 }
 
 FArchive* FNullNetworkReplayStreamer::GetMetadataArchive()
 {
-	return MetadataFileAr;
+	check( StreamerState != EStreamerState::Idle );
+
+	// Create the metadata archive on-demand
+	if (!MetadataFileAr)
+	{
+		switch (StreamerState)
+		{
+			case EStreamerState::Recording:
+				MetadataFileAr.Reset( IFileManager::Get().CreateFileWriter( *GetMetadataFilename(CurrentStreamName) ) );
+				break;
+
+			case EStreamerState::Playback:
+				MetadataFileAr.Reset( IFileManager::Get().CreateFileReader( *GetMetadataFilename(CurrentStreamName) ) );
+				break;
+
+			default:
+				break;
+		}
+	}
+
+	return MetadataFileAr.Get();
+}
+
+bool FNullNetworkReplayStreamer::IsLive( const FString& StreamName ) const
+{
+	// If the directory for this stream doesn't exist, it can't possibly be live.
+	if (!IFileManager::Get().DirectoryExists(*GetStreamDirectory(StreamName)))
+	{
+		return false;
+	}
+
+	// If the metadata file doesn't exist, this is a live stream.
+	const int64 MetadataFileSize = IFileManager::Get().FileSize(*GetMetadataFilename(StreamName));
+	return MetadataFileSize == INDEX_NONE;
+}
+
+void FNullNetworkReplayStreamer::DeleteFinishedStream( const FString& StreamName, const FOnDeleteFinishedStreamComplete& Delegate ) const
+{
+	// Live streams can't be deleted
+	if (IsLive(StreamName))
+	{
+		UE_LOG(LogNullReplay, Log, TEXT("Can't delete network replay stream %s because it is live!"), *StreamName);
+		Delegate.ExecuteIfBound(false);
+		return;
+	}
+
+	// Delete the directory with the specified name in the Saved/Demos directory
+	const FString DemoName = GetStreamDirectory(StreamName);
+
+	const bool DeleteSucceeded = IFileManager::Get().DeleteDirectory( *DemoName, false, true );
+
+	Delegate.ExecuteIfBound(DeleteSucceeded);
+}
+
+void FNullNetworkReplayStreamer::EnumerateStreams( const FOnEnumerateStreamsComplete& Delegate ) const
+{
+	// Simply returns a stream for each folder in the Saved/Demos directory
+	const FString WildCardPath = GetDemoPath() + TEXT( "*" );
+
+	TArray<FString> DirectoryNames;
+	IFileManager::Get().FindFiles( DirectoryNames, *WildCardPath, false, true );
+
+	TArray<FNetworkReplayStreamInfo> Results;
+
+	for (const FString& Directory : DirectoryNames)
+	{
+		// Assume there will be one file with a .demo extension in the directory
+		const FString FullDemoFilePath = GetDemoFilename(Directory);
+
+		FNetworkReplayStreamInfo Info;
+		Info.SizeInBytes = IFileManager::Get().FileSize( *FullDemoFilePath );
+		
+		if (Info.SizeInBytes != INDEX_NONE)
+		{
+			Info.Name = Directory;
+			Info.Timestamp = IFileManager::Get().GetTimeStamp( *FullDemoFilePath );
+			Info.bIsLive = IsLive( Directory );			
+
+			// Live streams not supported yet
+			if (!Info.bIsLive)
+			{
+				Results.Add(Info);
+			}
+		}
+	}
+
+	Delegate.ExecuteIfBound(Results);
 }
 
 IMPLEMENT_MODULE( FNullNetworkReplayStreamingFactory, NullNetworkReplayStreaming )
diff --git a/Engine/Source/Runtime/NetworkReplayStreaming/NullNetworkReplayStreaming/Public/NullNetworkReplayStreaming.h b/Engine/Source/Runtime/NetworkReplayStreaming/NullNetworkReplayStreaming/Public/NullNetworkReplayStreaming.h
--- a/Engine/Source/Runtime/NetworkReplayStreaming/NullNetworkReplayStreaming/Public/NullNetworkReplayStreaming.h
+++ b/Engine/Source/Runtime/NetworkReplayStreaming/NullNetworkReplayStreaming/Public/NullNetworkReplayStreaming.h
@@ -5,25 +5,47 @@
 #include "NetworkReplayStreaming.h"
 #include "Core.h"
 #include "ModuleManager.h"
+#include "UniquePtr.h"
 
 /** Default streamer that goes straight to the HD */
 class FNullNetworkReplayStreamer : public INetworkReplayStreamer
 {
 public:
-	FNullNetworkReplayStreamer() : FileAr( NULL ), MetadataFileAr( NULL ) {}
+	FNullNetworkReplayStreamer() :
+		StreamerState( EStreamerState::Idle )
+	{}
+	
+	/** INetworkReplayStreamer implementation */
 	virtual void StartStreaming( const FString& StreamName, bool bRecord, const FString& VersionString, const FOnStreamReadyDelegate& Delegate ) override;
 	virtual void StopStreaming() override;
 	virtual FArchive* GetHeaderArchive() override;
 	virtual FArchive* GetStreamingArchive() override;
 	virtual FArchive* GetMetadataArchive() override;
 	virtual bool IsDataAvailable() const override { return true; }
+	virtual bool IsLive( const FString& StreamName ) const override;
+	virtual void DeleteFinishedStream( const FString& StreamName, const FOnDeleteFinishedStreamComplete& Delegate) const override;
+	virtual void EnumerateStreams( const FOnEnumerateStreamsComplete& Delegate ) const override;
 
 private:
 	/** Handle to the archive that will read/write network packets */
-	FArchive* FileAr;
+	TUniquePtr<FArchive> FileAr;
 
 	/* Handle to the archive that will read/write metadata */
-	FArchive* MetadataFileAr;
+	TUniquePtr<FArchive> MetadataFileAr;
+
+	/** EStreamerState - Overall state of the streamer */
+	enum class EStreamerState
+	{
+		Idle,					// The streamer is idle. Either we haven't started streaming yet, or we are done
+		Recording,				// We are in the process of recording a replay to disk
+		Playback,				// We are in the process of playing a replay from disk
+	};
+
+	/** Overall state of the streamer */
+	EStreamerState StreamerState;
+
+	/** Remember the name of the current stream, if any. */
+	FString CurrentStreamName;
 };
 
 class FNullNetworkReplayStreamingFactory : public INetworkReplayStreamingFactory
