Engine side modifications necessary for http replay stream enumeration + cap max demo name in ShooterGame demo browser

diff --git a/Engine/Source/Runtime/NetworkReplayStreaming/HttpNetworkReplayStreaming/Private/HttpNetworkReplayStreaming.cpp b/Engine/Source/Runtime/NetworkReplayStreaming/HttpNetworkReplayStreaming/Private/HttpNetworkReplayStreaming.cpp
--- a/Engine/Source/Runtime/NetworkReplayStreaming/HttpNetworkReplayStreaming/Private/HttpNetworkReplayStreaming.cpp
+++ b/Engine/Source/Runtime/NetworkReplayStreaming/HttpNetworkReplayStreaming/Private/HttpNetworkReplayStreaming.cpp
@@ -50,6 +50,19 @@ void HttpStreamFArchive::Seek( int64 InPos )
 	Pos = InPos;
 }
 
+FHttpNetworkReplayStreamer::FHttpNetworkReplayStreamer() : 
+	StreamFileCount( 0 ), 
+	LastChunkTime( 0 ), 
+	StreamerState( EStreamerState::Idle ), 
+	HttpState( EHttptate::Idle ), 
+	bStopStreamingCalled( false ), 
+	bStreamIsLive( false ),
+	NumDownloadChunks( 0 ) 
+{
+	// Initialize the server URL
+	GConfig->GetString( TEXT( "HttpNetworkReplayStreaming" ), TEXT( "ServerURL" ), ServerURL, GEngineIni );
+}
+
 void FHttpNetworkReplayStreamer::StartStreaming( const FString& StreamName, bool bRecord, const FString& VersionString, const FOnStreamReadyDelegate& Delegate )
 {
 	if ( !SessionName.IsEmpty() )
@@ -73,7 +86,7 @@ void FHttpNetworkReplayStreamer::StartStreaming( const FString& StreamName, bool
 	SessionVersion = VersionString;
 
 	// Remember the delegate, which we'll call as soon as the header is available
-	RememberedDelegate = Delegate;
+	StartStreamingDelegate = Delegate;
 
 	// Setup the archives
 	StreamArchive.ArIsLoading = !bRecord;
@@ -92,9 +105,6 @@ void FHttpNetworkReplayStreamer::StartStreaming( const FString& StreamName, bool
 		SessionName = StreamName;
 	}
 
-	// Initialize the server URL
-	GConfig->GetString( TEXT( "HttpNetworkReplayStreaming" ), TEXT( "ServerURL" ), ServerURL, GEngineIni );
-
 	// Create the Http request and add to pending request list
 	TSharedRef<class IHttpRequest> HttpRequest = FHttpModule::Get().CreateRequest();
 
@@ -331,10 +341,26 @@ void FHttpNetworkReplayStreamer::DeleteFinishedStream( const FString& StreamName
 	Delegate.ExecuteIfBound(false);
 }
 
-void FHttpNetworkReplayStreamer::EnumerateStreams( const FOnEnumerateStreamsComplete& Delegate ) const
+void FHttpNetworkReplayStreamer::EnumerateStreams( const FString& VersionString, const FOnEnumerateStreamsComplete& Delegate )
 {
-	// Stubbed!
-	Delegate.ExecuteIfBound(TArray<FNetworkReplayStreamInfo>());
+	if ( IsHttpBusy() )
+	{
+		Delegate.ExecuteIfBound( TArray<FNetworkReplayStreamInfo>() );
+		return;
+	}
+
+	TSharedRef<class IHttpRequest> HttpRequest = FHttpModule::Get().CreateRequest();
+
+	// Enumerate all of the sessions
+	HttpRequest->SetURL( FString::Printf( TEXT( "%senumsessions?Version=%s" ), *ServerURL, *VersionString ) );
+	HttpRequest->SetVerb( TEXT( "GET" ) );
+
+	HttpRequest->OnProcessRequestComplete().BindRaw( this, &FHttpNetworkReplayStreamer::HttpEnumerateSessionsFinished );
+
+	EnumerateStreamsDelegate = Delegate;
+	HttpState = EHttptate::EnumeratingSessions;
+
+	HttpRequest->ProcessRequest();
 }
 
 void FHttpNetworkReplayStreamer::HttpStartUploadingFinished( FHttpRequestPtr HttpRequest, FHttpResponsePtr HttpResponse, bool bSucceeded )
@@ -466,10 +492,10 @@ void FHttpNetworkReplayStreamer::HttpStartDownloadingFinished( FHttpRequestPtr H
 		{
 			UE_LOG( LogHttpReplay, Warning, TEXT( "FHttpNetworkReplayStreamer::HttpStartDownloadingFinished. NO CHUNKS" ) );
 
-			RememberedDelegate.ExecuteIfBound( false, StreamArchive.IsSaving() );
+			StartStreamingDelegate.ExecuteIfBound( false, StreamArchive.IsSaving() );
 
 			// Reset delegate
-			RememberedDelegate = FOnStreamReadyDelegate();
+			StartStreamingDelegate = FOnStreamReadyDelegate();
 
 			StreamerState = EStreamerState::Idle;
 		}
@@ -479,10 +505,10 @@ void FHttpNetworkReplayStreamer::HttpStartDownloadingFinished( FHttpRequestPtr H
 		// FAIL
 		UE_LOG( LogHttpReplay, Warning, TEXT( "FHttpNetworkReplayStreamer::HttpStartDownloadingFinished. FAILED" ) );
 
-		RememberedDelegate.ExecuteIfBound( false, StreamArchive.IsSaving() );
+		StartStreamingDelegate.ExecuteIfBound( false, StreamArchive.IsSaving() );
 
 		// Reset delegate
-		RememberedDelegate = FOnStreamReadyDelegate();
+		StartStreamingDelegate = FOnStreamReadyDelegate();
 
 		StreamerState = EStreamerState::Idle;
 	}
@@ -502,7 +528,7 @@ void FHttpNetworkReplayStreamer::HttpDownloadHeaderFinished( FHttpRequestPtr Htt
 
 		UE_LOG( LogHttpReplay, Log, TEXT( "FHttpNetworkReplayStreamer::HttpDownloadHeaderFinished. Size: %i" ), HeaderArchive.Buffer.Num() );
 
-		RememberedDelegate.ExecuteIfBound( true, StreamArchive.IsSaving() );
+		StartStreamingDelegate.ExecuteIfBound( true, StreamArchive.IsSaving() );
 		StreamerState = EStreamerState::StreamingDown;
 	}
 	else
@@ -511,13 +537,13 @@ void FHttpNetworkReplayStreamer::HttpDownloadHeaderFinished( FHttpRequestPtr Htt
 		UE_LOG( LogHttpReplay, Warning, TEXT( "FHttpNetworkReplayStreamer::HttpDownloadFinished. FAILED." ) );
 
 		StreamArchive.Buffer.Empty();
-		RememberedDelegate.ExecuteIfBound( false, StreamArchive.IsSaving() );
+		StartStreamingDelegate.ExecuteIfBound( false, StreamArchive.IsSaving() );
 
 		StreamerState = EStreamerState::Idle;
 	}
 
 	// Reset delegate
-	RememberedDelegate = FOnStreamReadyDelegate();
+	StartStreamingDelegate = FOnStreamReadyDelegate();
 }
 
 void FHttpNetworkReplayStreamer::HttpDownloadFinished( FHttpRequestPtr HttpRequest, FHttpResponsePtr HttpResponse, bool bSucceeded )
@@ -565,6 +591,86 @@ void FHttpNetworkReplayStreamer::HttpDownloadFinished( FHttpRequestPtr HttpReque
 	}
 }
 
+void FHttpNetworkReplayStreamer::HttpEnumerateSessionsFinished( FHttpRequestPtr HttpRequest, FHttpResponsePtr HttpResponse, bool bSucceeded )
+{
+	check( HttpState == EHttptate::EnumeratingSessions );
+
+	HttpState = EHttptate::Idle;
+
+	if ( bSucceeded && HttpResponse->GetResponseCode() == EHttpResponseCodes::Ok )
+	{		
+		UE_LOG( LogHttpReplay, Log, TEXT( "FHttpNetworkReplayStreamer::HttpEnumerateSessionsFinished." ) );
+
+		TArray<FNetworkReplayStreamInfo> Streams;
+		FString StreamsString = HttpResponse->GetContentAsString();
+
+		int32 Index = INDEX_NONE;
+
+		TArray< FString > Tokens;
+
+		if ( StreamsString.FindChar( '{', Index ) )
+		{
+			StreamsString = StreamsString.RightChop( Index + 1 );
+
+			while ( StreamsString.FindChar( ',', Index ) )
+			{
+				Tokens.Add( StreamsString.Left( Index ) );
+				StreamsString = StreamsString.RightChop( Index + 1 );
+			}
+
+			if ( StreamsString.FindChar( '}', Index ) )
+			{
+				Tokens.Add( StreamsString.Left( Index ) );
+			}
+			else
+			{
+				UE_LOG( LogHttpReplay, Warning, TEXT( "FHttpNetworkReplayStreamer::HttpEnumerateSessionsFinished. '}' not found." ) );
+			}
+		}
+		else
+		{
+			UE_LOG( LogHttpReplay, Warning, TEXT( "FHttpNetworkReplayStreamer::HttpEnumerateSessionsFinished. '{' not found." ) );
+		}
+
+		for ( int i = 0; i < Tokens.Num(); i += 4 )
+		{
+			FNetworkReplayStreamInfo NewStream;
+
+			NewStream.Name			= Tokens[i];
+			NewStream.bIsLive		= false;
+			NewStream.SizeInBytes	= 0;
+			NewStream.Timestamp		= 0;
+
+			if ( i + 1 < Tokens.Num() )
+			{
+				// Server returns milliseconds from Jan 1 1970, GMT
+				NewStream.Timestamp = FDateTime( FCString::Atoi64( *Tokens[ i + 1 ] ) * 1000 * 10 + FDateTime( 1970, 1, 1 ).GetTicks() );
+			}
+
+			if ( i + 2 < Tokens.Num() )
+			{
+				NewStream.SizeInBytes = FCString::Atoi( *Tokens[ i + 2 ] );
+			}
+
+			if ( i + 3 < Tokens.Num() )
+			{
+				NewStream.bIsLive = Tokens[ i + 3 ].Contains( TEXT( "true" ) );
+			}
+
+			Streams.Add( NewStream );
+		}
+
+		EnumerateStreamsDelegate.ExecuteIfBound( Streams );
+	}
+	else
+	{
+		UE_LOG( LogHttpReplay, Warning, TEXT( "FHttpNetworkReplayStreamer::HttpEnumerateSessionsFinished. FAILED" ) );
+		EnumerateStreamsDelegate.ExecuteIfBound( TArray<FNetworkReplayStreamInfo>() );		// FIXME: Notify failure here
+	}
+
+	EnumerateStreamsDelegate = FOnEnumerateStreamsComplete();
+}
+
 void FHttpNetworkReplayStreamer::Tick( float DeltaTime )
 {
 	if ( IsHttpBusy() )
diff --git a/Engine/Source/Runtime/NetworkReplayStreaming/HttpNetworkReplayStreaming/Private/HttpNetworkReplayStreaming.h b/Engine/Source/Runtime/NetworkReplayStreaming/HttpNetworkReplayStreaming/Private/HttpNetworkReplayStreaming.h
--- a/Engine/Source/Runtime/NetworkReplayStreaming/HttpNetworkReplayStreaming/Private/HttpNetworkReplayStreaming.h
+++ b/Engine/Source/Runtime/NetworkReplayStreaming/HttpNetworkReplayStreaming/Private/HttpNetworkReplayStreaming.h
@@ -28,15 +28,7 @@ public:
 class FHttpNetworkReplayStreamer : public INetworkReplayStreamer
 {
 public:
-	FHttpNetworkReplayStreamer() : 
-		StreamFileCount( 0 ), 
-		LastChunkTime( 0 ), 
-		StreamerState( EStreamerState::Idle ), 
-		HttpState( EHttptate::Idle ), 
-		bStopStreamingCalled( false ), 
-		bStreamIsLive( false ),
-		NumDownloadChunks( 0 ) 
-	{}
+	FHttpNetworkReplayStreamer();
 
 	/** INetworkReplayStreamer implementation */
 	virtual void		StartStreaming( const FString& StreamName, bool bRecord, const FString& VersionString, const FOnStreamReadyDelegate& Delegate ) override;
@@ -47,7 +39,7 @@ public:
 	virtual bool		IsDataAvailable() const override;
 	virtual bool		IsLive( const FString& StreamName ) const override;
 	virtual void		DeleteFinishedStream( const FString& StreamName, const FOnDeleteFinishedStreamComplete& Delegate ) const override;
-	virtual void		EnumerateStreams( const FOnEnumerateStreamsComplete& Delegate ) const override;
+	virtual void		EnumerateStreams( const FString& VersionString, const FOnEnumerateStreamsComplete& Delegate ) override;
 
 	/** FHttpNetworkReplayStreamer */
 	void UploadHeader();
@@ -63,6 +55,7 @@ public:
 	void HttpStopUploadingFinished( FHttpRequestPtr HttpRequest, FHttpResponsePtr HttpResponse, bool bSucceeded );
 	void HttpHeaderUploadFinished( FHttpRequestPtr HttpRequest, FHttpResponsePtr HttpResponse, bool bSucceeded );
 	void HttpUploadFinished( FHttpRequestPtr HttpRequest, FHttpResponsePtr HttpResponse, bool bSucceeded );
+	void HttpEnumerateSessionsFinished( FHttpRequestPtr HttpRequest, FHttpResponsePtr HttpResponse, bool bSucceeded );
 
 	void Tick( float DeltaTime );
 	bool IsHttpBusy() const;		// True if we are waiting on http request response
@@ -79,6 +72,7 @@ public:
 		StartDownloading,			// We have made the request to start downloading a replay stream
 		DownloadingHeader,			// We are downloading the replay header
 		DownloadingStream,			// We are in the process of downloading the replay stream
+		EnumeratingSessions,		// We are in the process of downloading the available sessions
 	};
 
 	/** EStreamerState - Overall state of the streamer */
@@ -92,7 +86,6 @@ public:
 		StreamingUpFinal,			// We are uploading the final stream
 	};
 
-	FOnStreamReadyDelegate	RememberedDelegate;		// Delegate passed in to StartStreaming
 	HttpStreamFArchive		HeaderArchive;			// Archive used to buffer the header stream
 	HttpStreamFArchive		StreamArchive;			// Archive used to buffer the data stream
 	FString					SessionName;			// Name of the session on the http replay server
@@ -105,6 +98,9 @@ public:
 	bool					bStopStreamingCalled;
 	bool					bStreamIsLive;			// If true, we are viewing a live stream
 	int32					NumDownloadChunks;
+
+	FOnStreamReadyDelegate			StartStreamingDelegate;		// Delegate passed in to StartStreaming
+	FOnEnumerateStreamsComplete		EnumerateStreamsDelegate;
 };
 
 class FHttpNetworkReplayStreamingFactory : public INetworkReplayStreamingFactory, public FTickableGameObject
diff --git a/Engine/Source/Runtime/NetworkReplayStreaming/NetworkReplayStreaming/Public/NetworkReplayStreaming.h b/Engine/Source/Runtime/NetworkReplayStreaming/NetworkReplayStreaming/Public/NetworkReplayStreaming.h
--- a/Engine/Source/Runtime/NetworkReplayStreaming/NetworkReplayStreaming/Public/NetworkReplayStreaming.h
+++ b/Engine/Source/Runtime/NetworkReplayStreaming/NetworkReplayStreaming/Public/NetworkReplayStreaming.h
@@ -68,7 +68,7 @@ public:
 	 *
 	 * @param Delegate A delegate that will be executed if bound when the list of streams is available
 	 */
-	virtual void EnumerateStreams( const FOnEnumerateStreamsComplete& Delegate ) const = 0;
+	virtual void EnumerateStreams( const FString& VersionString, const FOnEnumerateStreamsComplete& Delegate ) = 0;
 };
 
 /** Replay streamer factory */
diff --git a/Engine/Source/Runtime/NetworkReplayStreaming/NullNetworkReplayStreaming/Private/NullNetworkReplayStreaming.cpp b/Engine/Source/Runtime/NetworkReplayStreaming/NullNetworkReplayStreaming/Private/NullNetworkReplayStreaming.cpp
--- a/Engine/Source/Runtime/NetworkReplayStreaming/NullNetworkReplayStreaming/Private/NullNetworkReplayStreaming.cpp
+++ b/Engine/Source/Runtime/NetworkReplayStreaming/NullNetworkReplayStreaming/Private/NullNetworkReplayStreaming.cpp
@@ -169,7 +169,7 @@ void FNullNetworkReplayStreamer::DeleteFinishedStream( const FString& StreamName
 	Delegate.ExecuteIfBound(DeleteSucceeded);
 }
 
-void FNullNetworkReplayStreamer::EnumerateStreams( const FOnEnumerateStreamsComplete& Delegate ) const
+void FNullNetworkReplayStreamer::EnumerateStreams( const FString& VersionString, const FOnEnumerateStreamsComplete& Delegate )
 {
 	// Simply returns a stream for each folder in the Saved/Demos directory
 	const FString WildCardPath = GetDemoPath() + TEXT( "*" );
diff --git a/Engine/Source/Runtime/NetworkReplayStreaming/NullNetworkReplayStreaming/Public/NullNetworkReplayStreaming.h b/Engine/Source/Runtime/NetworkReplayStreaming/NullNetworkReplayStreaming/Public/NullNetworkReplayStreaming.h
--- a/Engine/Source/Runtime/NetworkReplayStreaming/NullNetworkReplayStreaming/Public/NullNetworkReplayStreaming.h
+++ b/Engine/Source/Runtime/NetworkReplayStreaming/NullNetworkReplayStreaming/Public/NullNetworkReplayStreaming.h
@@ -24,7 +24,7 @@ public:
 	virtual bool IsDataAvailable() const override { return true; }
 	virtual bool IsLive( const FString& StreamName ) const override;
 	virtual void DeleteFinishedStream( const FString& StreamName, const FOnDeleteFinishedStreamComplete& Delegate) const override;
-	virtual void EnumerateStreams( const FOnEnumerateStreamsComplete& Delegate ) const override;
+	virtual void EnumerateStreams( const FString& VersionString, const FOnEnumerateStreamsComplete& Delegate ) override;
 
 private:
 	/** Handle to the archive that will read/write network packets */
