Remove white space from replay names during enumeration + more robust error handling

diff --git a/Engine/Source/Runtime/NetworkReplayStreaming/HttpNetworkReplayStreaming/Private/HttpNetworkReplayStreaming.cpp b/Engine/Source/Runtime/NetworkReplayStreaming/HttpNetworkReplayStreaming/Private/HttpNetworkReplayStreaming.cpp
--- a/Engine/Source/Runtime/NetworkReplayStreaming/HttpNetworkReplayStreaming/Private/HttpNetworkReplayStreaming.cpp
+++ b/Engine/Source/Runtime/NetworkReplayStreaming/HttpNetworkReplayStreaming/Private/HttpNetworkReplayStreaming.cpp
@@ -608,6 +608,8 @@ void FHttpNetworkReplayStreamer::HttpEnumerateSessionsFinished( FHttpRequestPtr
 
 		TArray< FString > Tokens;
 
+		// Parse the string as { token 1, token ..., token n }
+		// This isn't perfect, we should convert to JSON when the dust settles
 		if ( StreamsString.FindChar( '{', Index ) )
 		{
 			StreamsString = StreamsString.RightChop( Index + 1 );
@@ -625,14 +627,37 @@ void FHttpNetworkReplayStreamer::HttpEnumerateSessionsFinished( FHttpRequestPtr
 			else
 			{
 				UE_LOG( LogHttpReplay, Warning, TEXT( "FHttpNetworkReplayStreamer::HttpEnumerateSessionsFinished. '}' not found." ) );
+				EnumerateStreamsDelegate.ExecuteIfBound( TArray<FNetworkReplayStreamInfo>() );		// FIXME: Notify failure here
+				EnumerateStreamsDelegate = FOnEnumerateStreamsComplete();
+				return;
+			}
+
+			if ( Tokens.Num() > 0 )
+			{
+				Tokens[ Tokens.Num() - 1 ].RemoveFromStart( TEXT( " " ) );
+				Tokens[ Tokens.Num() - 1 ].RemoveFromEnd( TEXT( " " ) );
 			}
 		}
 		else
 		{
 			UE_LOG( LogHttpReplay, Warning, TEXT( "FHttpNetworkReplayStreamer::HttpEnumerateSessionsFinished. '{' not found." ) );
+			EnumerateStreamsDelegate.ExecuteIfBound( TArray<FNetworkReplayStreamInfo>() );		// FIXME: Notify failure here
+			EnumerateStreamsDelegate = FOnEnumerateStreamsComplete();
+			return;
+		}
+
+		const int NUM_TOKENS_PER_INFO = 4;
+
+		if ( Tokens.Num() == 0 || ( Tokens.Num() % NUM_TOKENS_PER_INFO ) != 0 )
+		{
+			UE_LOG( LogHttpReplay, Warning, TEXT( "FHttpNetworkReplayStreamer::HttpEnumerateSessionsFinished. Invalid number of tokens: %i" ), Tokens.Num() );
+			EnumerateStreamsDelegate.ExecuteIfBound( TArray<FNetworkReplayStreamInfo>() );		// FIXME: Notify failure here
+			EnumerateStreamsDelegate = FOnEnumerateStreamsComplete();
+			return;
 		}
 
-		for ( int i = 0; i < Tokens.Num(); i += 4 )
+		// Convert tokens to individual FNetworkReplayStreamInfo's
+		for ( int i = 0; i < Tokens.Num(); i += NUM_TOKENS_PER_INFO )
 		{
 			FNetworkReplayStreamInfo NewStream;
 
@@ -643,7 +668,8 @@ void FHttpNetworkReplayStreamer::HttpEnumerateSessionsFinished( FHttpRequestPtr
 
 			if ( i + 1 < Tokens.Num() )
 			{
-				// Server returns milliseconds from Jan 1 1970, GMT
+				// Server returns milliseconds from January 1, 1970, 00:00:00 GMT
+				// We need to compensate for the fact that FDateTime starts at January 1, 0001 A.D. and is in 100 nanosecond resolution
 				NewStream.Timestamp = FDateTime( FCString::Atoi64( *Tokens[ i + 1 ] ) * 1000 * 10 + FDateTime( 1970, 1, 1 ).GetTicks() );
 			}
 
