Add ability for network replay interface to track total demo time during live recording and playback

diff --git a/Engine/Source/Runtime/Engine/Private/DemoNetDriver.cpp b/Engine/Source/Runtime/Engine/Private/DemoNetDriver.cpp
--- a/Engine/Source/Runtime/Engine/Private/DemoNetDriver.cpp
+++ b/Engine/Source/Runtime/Engine/Private/DemoNetDriver.cpp
@@ -615,6 +615,8 @@ void UDemoNetDriver::TickDemoRecord( float DeltaSeconds )
 	DemoDeltaTime += DeltaSeconds;
 	DemoCurrentTime += DeltaSeconds;
 
+	ReplayStreamer->UpdateTotalDemoTime( DemoCurrentTime * 1000 );
+
 	const double CurrentSeconds = FPlatformTime::Seconds();
 
 	const double RECORD_HZ		= CVarDemoRecordHz.GetValueOnGameThread();
@@ -688,6 +690,13 @@ bool UDemoNetDriver::ReadDemoFrame()
 		return false;
 	}
 
+	const uint32 TotalDemoTimeInMS = ReplayStreamer->GetTotalDemoTime();
+
+	if ( TotalDemoTimeInMS > 0 )
+	{
+		DemoTotalTime = (float)TotalDemoTimeInMS / 1000.0f;
+	}
+
 	FArchive* FileAr = ReplayStreamer->GetStreamingArchive();
 
 	if ( FileAr->IsError() )
diff --git a/Engine/Source/Runtime/NetworkReplayStreaming/HttpNetworkReplayStreaming/Private/HttpNetworkReplayStreaming.cpp b/Engine/Source/Runtime/NetworkReplayStreaming/HttpNetworkReplayStreaming/Private/HttpNetworkReplayStreaming.cpp
--- a/Engine/Source/Runtime/NetworkReplayStreaming/HttpNetworkReplayStreaming/Private/HttpNetworkReplayStreaming.cpp
+++ b/Engine/Source/Runtime/NetworkReplayStreaming/HttpNetworkReplayStreaming/Private/HttpNetworkReplayStreaming.cpp
@@ -57,7 +57,8 @@ FHttpNetworkReplayStreamer::FHttpNetworkReplayStreamer() :
 	HttpState( EHttptate::Idle ), 
 	bStopStreamingCalled( false ), 
 	bStreamIsLive( false ),
-	NumDownloadChunks( 0 ) 
+	NumDownloadChunks( 0 ),
+	DemoTimeInMS( 0 )
 {
 	// Initialize the server URL
 	GConfig->GetString( TEXT( "HttpNetworkReplayStreaming" ), TEXT( "ServerURL" ), ServerURL, GEngineIni );
@@ -196,7 +197,7 @@ void FHttpNetworkReplayStreamer::UploadHeader()
 
 	HttpState = EHttptate::UploadingHeader;
 
-	HttpRequest->SetURL( FString::Printf( TEXT( "%supload?Version=%s&Session=%s&NumChunks=%i&Filename=replay.header" ), *ServerURL, *SessionVersion, *SessionName, StreamFileCount ) );
+	HttpRequest->SetURL( FString::Printf( TEXT( "%supload?Version=%s&Session=%s&NumChunks=%i&Time=%i&Filename=replay.header" ), *ServerURL, *SessionVersion, *SessionName, StreamFileCount, DemoTimeInMS ) );
 	HttpRequest->SetVerb( TEXT( "POST" ) );
 	HttpRequest->SetHeader( TEXT( "Content-Type" ), TEXT( "application/octet-stream" ) );
 	HttpRequest->SetContent( HeaderArchive.Buffer );
@@ -234,7 +235,7 @@ void FHttpNetworkReplayStreamer::FlushStream()
 
 	HttpState = EHttptate::UploadingStream;
 
-	HttpRequest->SetURL( FString::Printf( TEXT( "%supload?Version=%s&Session=%s&NumChunks=%i&Filename=stream.%i" ), *ServerURL, *SessionVersion, *SessionName, StreamFileCount + 1, StreamFileCount ) );
+	HttpRequest->SetURL( FString::Printf( TEXT( "%supload?Version=%s&Session=%s&NumChunks=%i&Time=%i&Filename=stream.%i" ), *ServerURL, *SessionVersion, *SessionName, StreamFileCount + 1, DemoTimeInMS, StreamFileCount ) );
 	HttpRequest->SetVerb( TEXT( "POST" ) );
 	HttpRequest->SetHeader( TEXT( "Content-Type" ), TEXT( "application/octet-stream" ) );
 	HttpRequest->SetContent( StreamArchive.Buffer );
@@ -318,6 +319,11 @@ FArchive* FHttpNetworkReplayStreamer::GetMetadataArchive()
 	return NULL;
 }
 
+void FHttpNetworkReplayStreamer::UpdateTotalDemoTime( uint32 TimeInMS )
+{
+	DemoTimeInMS = TimeInMS;
+}
+
 bool FHttpNetworkReplayStreamer::IsDataAvailable() const
 {
 	// If we are loading, and we have more data
@@ -451,7 +457,7 @@ void FHttpNetworkReplayStreamer::HttpUploadFinished( FHttpRequestPtr HttpRequest
 
 			HttpState = EHttptate::StopUploading;
 
-			HttpRequest->SetURL( FString::Printf( TEXT( "%sstopuploading?Version=%s&Session=%s&NumChunks=%i" ), *ServerURL, *SessionVersion, *SessionName, StreamFileCount ) );
+			HttpRequest->SetURL( FString::Printf( TEXT( "%sstopuploading?Version=%s&Session=%s&NumChunks=%i&Time=%i" ), *ServerURL, *SessionVersion, *SessionName, StreamFileCount, DemoTimeInMS ) );
 			HttpRequest->SetVerb( TEXT( "POST" ) );
 			HttpRequest->SetHeader( TEXT( "Content-Type" ), TEXT( "application/octet-stream" ) );
 
@@ -475,13 +481,15 @@ void FHttpNetworkReplayStreamer::HttpStartDownloadingFinished( FHttpRequestPtr H
 	if ( bSucceeded && HttpResponse->GetResponseCode() == EHttpResponseCodes::Ok )
 	{
 		FString NumChunksString = HttpResponse->GetHeader( TEXT( "NumChunks" ) );
+		FString DemoTimeString = HttpResponse->GetHeader( TEXT( "Time" ) );
 		FString State = HttpResponse->GetHeader( TEXT( "State" ) );
 
 		bStreamIsLive = State == TEXT( "Live" );
 
 		NumDownloadChunks = FCString::Atoi( *NumChunksString );
+		DemoTimeInMS = FCString::Atoi( *DemoTimeString );
 
-		UE_LOG( LogHttpReplay, Log, TEXT( "FHttpNetworkReplayStreamer::HttpStartDownloadingFinished. State: %s, NumChunks: %i" ), *State, NumDownloadChunks );
+		UE_LOG( LogHttpReplay, Log, TEXT( "FHttpNetworkReplayStreamer::HttpStartDownloadingFinished. State: %s, NumChunks: %i, DemoTime: %2.2f" ), *State, NumDownloadChunks, (float)DemoTimeInMS / 1000 );
 
 		// First, download the header
 		if ( NumDownloadChunks > 0 )
@@ -556,11 +564,13 @@ void FHttpNetworkReplayStreamer::HttpDownloadFinished( FHttpRequestPtr HttpReque
 	if ( bSucceeded && HttpResponse->GetResponseCode() == EHttpResponseCodes::Ok )
 	{
 		FString NumChunksString = HttpResponse->GetHeader( TEXT( "NumChunks" ) );
+		FString DemoTimeString = HttpResponse->GetHeader( TEXT( "Time" ) );
 		FString State = HttpResponse->GetHeader( TEXT( "State" ) );
 
 		bStreamIsLive = State == TEXT( "Live" );
 
 		NumDownloadChunks = FCString::Atoi( *NumChunksString );
+		DemoTimeInMS = FCString::Atoi( *DemoTimeString );
 
 		if ( HttpResponse->GetContent().Num() > 0 || bStreamIsLive )
 		{
@@ -570,7 +580,7 @@ void FHttpNetworkReplayStreamer::HttpDownloadFinished( FHttpRequestPtr HttpReque
 				StreamFileCount++;
 			}
 
-			UE_LOG( LogHttpReplay, Log, TEXT( "FHttpNetworkReplayStreamer::HttpDownloadFinished. State: %s, Progress: %i / %i" ), *State, StreamFileCount, NumDownloadChunks );
+			UE_LOG( LogHttpReplay, Log, TEXT( "FHttpNetworkReplayStreamer::HttpDownloadFinished. State: %s, Progress: %i / %i, DemoTime: %2.2f" ), *State, StreamFileCount, NumDownloadChunks, (float)DemoTimeInMS / 1000 );
 		}
 		else
 		{
diff --git a/Engine/Source/Runtime/NetworkReplayStreaming/HttpNetworkReplayStreaming/Private/HttpNetworkReplayStreaming.h b/Engine/Source/Runtime/NetworkReplayStreaming/HttpNetworkReplayStreaming/Private/HttpNetworkReplayStreaming.h
--- a/Engine/Source/Runtime/NetworkReplayStreaming/HttpNetworkReplayStreaming/Private/HttpNetworkReplayStreaming.h
+++ b/Engine/Source/Runtime/NetworkReplayStreaming/HttpNetworkReplayStreaming/Private/HttpNetworkReplayStreaming.h
@@ -36,6 +36,8 @@ public:
 	virtual FArchive*	GetHeaderArchive() override;
 	virtual FArchive*	GetStreamingArchive() override;
 	virtual FArchive*	GetMetadataArchive() override;
+	virtual void		UpdateTotalDemoTime( uint32 TimeInMS ) override;
+	virtual uint32		GetTotalDemoTime() const override { return DemoTimeInMS; }
 	virtual bool		IsDataAvailable() const override;
 	virtual bool		IsLive( const FString& StreamName ) const override;
 	virtual void		DeleteFinishedStream( const FString& StreamName, const FOnDeleteFinishedStreamComplete& Delegate ) const override;
@@ -98,6 +100,7 @@ public:
 	bool					bStopStreamingCalled;
 	bool					bStreamIsLive;			// If true, we are viewing a live stream
 	int32					NumDownloadChunks;
+	uint32					DemoTimeInMS;
 
 	FOnStreamReadyDelegate			StartStreamingDelegate;		// Delegate passed in to StartStreaming
 	FOnEnumerateStreamsComplete		EnumerateStreamsDelegate;
diff --git a/Engine/Source/Runtime/NetworkReplayStreaming/NetworkReplayStreaming/Public/NetworkReplayStreaming.h b/Engine/Source/Runtime/NetworkReplayStreaming/NetworkReplayStreaming/Public/NetworkReplayStreaming.h
--- a/Engine/Source/Runtime/NetworkReplayStreaming/NetworkReplayStreaming/Public/NetworkReplayStreaming.h
+++ b/Engine/Source/Runtime/NetworkReplayStreaming/NetworkReplayStreaming/Public/NetworkReplayStreaming.h
@@ -50,6 +50,8 @@ public:
 	virtual FArchive* GetHeaderArchive() = 0;
 	virtual FArchive* GetStreamingArchive() = 0;
 	virtual FArchive* GetMetadataArchive() = 0;
+	virtual void UpdateTotalDemoTime( uint32 TimeInMS ) = 0;
+	virtual uint32 GetTotalDemoTime() const = 0;
 	virtual bool IsDataAvailable() const = 0;
 	
 	/** Returns true if the given StreamName is a game currently in progress */
diff --git a/Engine/Source/Runtime/NetworkReplayStreaming/NullNetworkReplayStreaming/Public/NullNetworkReplayStreaming.h b/Engine/Source/Runtime/NetworkReplayStreaming/NullNetworkReplayStreaming/Public/NullNetworkReplayStreaming.h
--- a/Engine/Source/Runtime/NetworkReplayStreaming/NullNetworkReplayStreaming/Public/NullNetworkReplayStreaming.h
+++ b/Engine/Source/Runtime/NetworkReplayStreaming/NullNetworkReplayStreaming/Public/NullNetworkReplayStreaming.h
@@ -21,6 +21,8 @@ public:
 	virtual FArchive* GetHeaderArchive() override;
 	virtual FArchive* GetStreamingArchive() override;
 	virtual FArchive* GetMetadataArchive() override;
+	virtual void UpdateTotalDemoTime( uint32 TimeInMS ) override { }
+	virtual uint32 GetTotalDemoTime() const override { return 0; }
 	virtual bool IsDataAvailable() const override { return true; }
 	virtual bool IsLive( const FString& StreamName ) const override;
 	virtual void DeleteFinishedStream( const FString& StreamName, const FOnDeleteFinishedStreamComplete& Delegate) const override;
