Engine side support for tracking replay viewers

diff --git a/Engine/Source/Runtime/NetworkReplayStreaming/HttpNetworkReplayStreaming/Private/HttpNetworkReplayStreaming.cpp b/Engine/Source/Runtime/NetworkReplayStreaming/HttpNetworkReplayStreaming/Private/HttpNetworkReplayStreaming.cpp
--- a/Engine/Source/Runtime/NetworkReplayStreaming/HttpNetworkReplayStreaming/Private/HttpNetworkReplayStreaming.cpp
+++ b/Engine/Source/Runtime/NetworkReplayStreaming/HttpNetworkReplayStreaming/Private/HttpNetworkReplayStreaming.cpp
@@ -53,6 +53,7 @@ void HttpStreamFArchive::Seek( int64 InPos )
 FHttpNetworkReplayStreamer::FHttpNetworkReplayStreamer() : 
 	StreamFileCount( 0 ), 
 	LastChunkTime( 0 ), 
+	LastRefreshViewerTime( 0 ),
 	StreamerState( EStreamerState::Idle ), 
 	HttpState( EHttptate::Idle ), 
 	bStopStreamingCalled( false ), 
@@ -255,7 +256,7 @@ void FHttpNetworkReplayStreamer::DownloadHeader()
 	// Download header first
 	TSharedRef<class IHttpRequest> HttpRequest = FHttpModule::Get().CreateRequest();
 
-	HttpRequest->SetURL( FString::Printf( TEXT( "%sdownload?Version=%s&Session=%s&Filename=replay.header" ), *ServerURL, *SessionVersion, *SessionName ) );
+	HttpRequest->SetURL( FString::Printf( TEXT( "%sdownload?Version=%s&Session=%s&Filename=replay.header" ), *ServerURL, *SessionVersion, *SessionName, *ViewerName ) );
 	HttpRequest->SetVerb( TEXT( "GET" ) );
 
 	HttpRequest->OnProcessRequestComplete().BindRaw( this, &FHttpNetworkReplayStreamer::HttpDownloadHeaderFinished );
@@ -267,13 +268,18 @@ void FHttpNetworkReplayStreamer::DownloadHeader()
 
 void FHttpNetworkReplayStreamer::DownloadNextChunk()
 {
+	if ( IsHttpBusy() )
+	{
+		return;
+	}
+
 	const bool bMoreChunksDefinitelyAvilable = StreamFileCount < NumDownloadChunks;
 	const bool bMoreChunksMayBeAvilable = bStreamIsLive;
 
 	if ( !bMoreChunksDefinitelyAvilable && !bMoreChunksMayBeAvilable )
 	{
-		StreamerState = EStreamerState::Idle;
-		UE_LOG( LogHttpReplay, Log, TEXT( "FHttpNetworkReplayStreamer::DownloadNextChunk. Done streaming." ) );
+		//StreamerState = EStreamerState::Idle;
+		//UE_LOG( LogHttpReplay, Log, TEXT( "FHttpNetworkReplayStreamer::DownloadNextChunk. Done streaming." ) );
 		return;
 	}
 
@@ -304,6 +310,46 @@ void FHttpNetworkReplayStreamer::DownloadNextChunk()
 	LastChunkTime = FPlatformTime::Seconds();
 }
 
+void FHttpNetworkReplayStreamer::RefreshViewer()
+{
+	if ( IsHttpBusy() )
+	{
+		return;
+	}
+
+	// Determine if it's time to download the next chunk
+	const double REFRESH_VIEWER_IN_SECONDS = 5;
+
+	if ( StreamerState != EStreamerState::StreamingDownFinal && FPlatformTime::Seconds() - LastRefreshViewerTime < REFRESH_VIEWER_IN_SECONDS )
+	{
+		// Not time yet
+		return;
+	}
+
+	TSharedRef<class IHttpRequest> HttpRequest = FHttpModule::Get().CreateRequest();
+
+	// Download the next stream chunk
+	if ( StreamerState == EStreamerState::StreamingDownFinal )
+	{
+		HttpRequest->SetURL( FString::Printf( TEXT( "%srefreshviewer?Version=%s&Session=%s&Viewer=%s&Final=true" ), *ServerURL, *SessionVersion, *SessionName, *ViewerName ) );
+		StreamerState = EStreamerState::Idle;
+	}
+	else
+	{
+		HttpRequest->SetURL( FString::Printf( TEXT( "%srefreshviewer?Version=%s&Session=%s&Viewer=%s" ), *ServerURL, *SessionVersion, *SessionName, *ViewerName ) );
+	}
+
+	HttpRequest->SetVerb( TEXT( "POST" ) );
+
+	HttpRequest->OnProcessRequestComplete().BindRaw( this, &FHttpNetworkReplayStreamer::HttpRefreshViewerFinished );
+
+	HttpState = EHttptate::RefreshingViewer;
+
+	HttpRequest->ProcessRequest();
+
+	LastRefreshViewerTime = FPlatformTime::Seconds();
+}
+
 FArchive* FHttpNetworkReplayStreamer::GetHeaderArchive()
 {
 	return ( HeaderArchive.IsSaving() || HeaderArchive.Pos < HeaderArchive.Buffer.Num() ) ? &HeaderArchive : NULL;
@@ -484,12 +530,14 @@ void FHttpNetworkReplayStreamer::HttpStartDownloadingFinished( FHttpRequestPtr H
 		FString DemoTimeString = HttpResponse->GetHeader( TEXT( "Time" ) );
 		FString State = HttpResponse->GetHeader( TEXT( "State" ) );
 
+		ViewerName = HttpResponse->GetHeader( TEXT( "Viewer" ) );
+
 		bStreamIsLive = State == TEXT( "Live" );
 
 		NumDownloadChunks = FCString::Atoi( *NumChunksString );
 		DemoTimeInMS = FCString::Atoi( *DemoTimeString );
 
-		UE_LOG( LogHttpReplay, Log, TEXT( "FHttpNetworkReplayStreamer::HttpStartDownloadingFinished. State: %s, NumChunks: %i, DemoTime: %2.2f" ), *State, NumDownloadChunks, (float)DemoTimeInMS / 1000 );
+		UE_LOG( LogHttpReplay, Log, TEXT( "FHttpNetworkReplayStreamer::HttpStartDownloadingFinished. Viewer: %s, State: %s, NumChunks: %i, DemoTime: %2.2f" ), *ViewerName, *State, NumDownloadChunks, (float)DemoTimeInMS / 1000 );
 
 		// First, download the header
 		if ( NumDownloadChunks > 0 )
@@ -542,7 +590,7 @@ void FHttpNetworkReplayStreamer::HttpDownloadHeaderFinished( FHttpRequestPtr Htt
 	else
 	{
 		// FAIL
-		UE_LOG( LogHttpReplay, Warning, TEXT( "FHttpNetworkReplayStreamer::HttpDownloadFinished. FAILED." ) );
+		UE_LOG( LogHttpReplay, Warning, TEXT( "FHttpNetworkReplayStreamer::HttpDownloadHeaderFinished. FAILED." ) );
 
 		StreamArchive.Buffer.Empty();
 		StartStreamingDelegate.ExecuteIfBound( false, StreamArchive.IsSaving() );
@@ -601,6 +649,20 @@ void FHttpNetworkReplayStreamer::HttpDownloadFinished( FHttpRequestPtr HttpReque
 	}
 }
 
+void FHttpNetworkReplayStreamer::HttpRefreshViewerFinished( FHttpRequestPtr HttpRequest, FHttpResponsePtr HttpResponse, bool bSucceeded )
+{
+	check( HttpState == EHttptate::RefreshingViewer );
+	check( StreamArchive.IsLoading() );
+
+	HttpState = EHttptate::Idle;
+
+	if ( !bSucceeded || HttpResponse->GetResponseCode() != EHttpResponseCodes::Ok )
+	{
+		// FAIL
+		UE_LOG( LogHttpReplay, Warning, TEXT( "FHttpNetworkReplayStreamer::HttpRefreshViewerFinished. FAILED." ) );
+	}
+}
+
 void FHttpNetworkReplayStreamer::HttpEnumerateSessionsFinished( FHttpRequestPtr HttpRequest, FHttpResponsePtr HttpResponse, bool bSucceeded )
 {
 	check( HttpState == EHttptate::EnumeratingSessions );
@@ -730,7 +792,7 @@ void FHttpNetworkReplayStreamer::Tick( float DeltaTime )
 		}
 		else if ( StreamerState == EStreamerState::StreamingDown )
 		{
-			StreamerState = EStreamerState::Idle;	// FIXME: Implement StopDownloading
+			StreamerState = EStreamerState::StreamingDownFinal;
 			bStopStreamingCalled = false;
 		}
 	}
@@ -761,6 +823,11 @@ void FHttpNetworkReplayStreamer::Tick( float DeltaTime )
 	else if ( StreamerState == EStreamerState::StreamingDown )
 	{
 		DownloadNextChunk();
+		RefreshViewer();
+	}
+	else if ( StreamerState == EStreamerState::StreamingDownFinal )
+	{
+		RefreshViewer();
 	}
 }
 
diff --git a/Engine/Source/Runtime/NetworkReplayStreaming/HttpNetworkReplayStreaming/Private/HttpNetworkReplayStreaming.h b/Engine/Source/Runtime/NetworkReplayStreaming/HttpNetworkReplayStreaming/Private/HttpNetworkReplayStreaming.h
--- a/Engine/Source/Runtime/NetworkReplayStreaming/HttpNetworkReplayStreaming/Private/HttpNetworkReplayStreaming.h
+++ b/Engine/Source/Runtime/NetworkReplayStreaming/HttpNetworkReplayStreaming/Private/HttpNetworkReplayStreaming.h
@@ -48,11 +48,13 @@ public:
 	void FlushStream();
 	void DownloadHeader();
 	void DownloadNextChunk();
+	void RefreshViewer();
 
 	/** Delegates */
 	void HttpStartDownloadingFinished( FHttpRequestPtr HttpRequest, FHttpResponsePtr HttpResponse, bool bSucceeded );
 	void HttpDownloadHeaderFinished( FHttpRequestPtr HttpRequest, FHttpResponsePtr HttpResponse, bool bSucceeded );
 	void HttpDownloadFinished( FHttpRequestPtr HttpRequest, FHttpResponsePtr HttpResponse, bool bSucceeded );
+	void HttpRefreshViewerFinished( FHttpRequestPtr HttpRequest, FHttpResponsePtr HttpResponse, bool bSucceeded );
 	void HttpStartUploadingFinished( FHttpRequestPtr HttpRequest, FHttpResponsePtr HttpResponse, bool bSucceeded );
 	void HttpStopUploadingFinished( FHttpRequestPtr HttpRequest, FHttpResponsePtr HttpResponse, bool bSucceeded );
 	void HttpHeaderUploadFinished( FHttpRequestPtr HttpRequest, FHttpResponsePtr HttpResponse, bool bSucceeded );
@@ -74,6 +76,7 @@ public:
 		StartDownloading,			// We have made the request to start downloading a replay stream
 		DownloadingHeader,			// We are downloading the replay header
 		DownloadingStream,			// We are in the process of downloading the replay stream
+		RefreshingViewer,			// We are refreshing the server to let it know we're still viewing
 		EnumeratingSessions,		// We are in the process of downloading the available sessions
 	};
 
@@ -85,6 +88,7 @@ public:
 		NeedToDownloadHeader,		// We are waiting to download the header
 		StreamingUp,				// We are in the process of streaming a replay to the http server
 		StreamingDown,				// We are in the process of streaming a replay from the http server
+		StreamingDownFinal,			// We are done streaming (will notify server to remove us as a viewer)
 		StreamingUpFinal,			// We are uploading the final stream
 	};
 
@@ -95,12 +99,14 @@ public:
 	FString					ServerURL;				// The address of the server
 	int32					StreamFileCount;		// Used as a counter to increment the stream.x extension count
 	double					LastChunkTime;			// The last time we uploaded/downloaded a chunk
+	double					LastRefreshViewerTime;	// The last time we refreshed ourselves as an active viewer
 	EStreamerState			StreamerState;			// Overall state of the streamer
 	EHttptate				HttpState;				// If not idle, there is a http request in flight
 	bool					bStopStreamingCalled;
 	bool					bStreamIsLive;			// If true, we are viewing a live stream
 	int32					NumDownloadChunks;
 	uint32					DemoTimeInMS;
+	FString					ViewerName;
 
 	FOnStreamReadyDelegate			StartStreamingDelegate;		// Delegate passed in to StartStreaming
 	FOnEnumerateStreamsComplete		EnumerateStreamsDelegate;
