Implement FHttpNetworkReplayStreamer::IsLive

diff --git a/Engine/Source/Runtime/NetworkReplayStreaming/HttpNetworkReplayStreaming/Private/HttpNetworkReplayStreaming.cpp b/Engine/Source/Runtime/NetworkReplayStreaming/HttpNetworkReplayStreaming/Private/HttpNetworkReplayStreaming.cpp
--- a/Engine/Source/Runtime/NetworkReplayStreaming/HttpNetworkReplayStreaming/Private/HttpNetworkReplayStreaming.cpp
+++ b/Engine/Source/Runtime/NetworkReplayStreaming/HttpNetworkReplayStreaming/Private/HttpNetworkReplayStreaming.cpp
@@ -383,8 +383,7 @@ bool FHttpNetworkReplayStreamer::IsDataAvailable() const
 
 bool FHttpNetworkReplayStreamer::IsLive(const FString& StreamName) const 
 {
-	// Stubbed!
-	return false;
+	return bStreamIsLive;
 }
 
 void FHttpNetworkReplayStreamer::DeleteFinishedStream( const FString& StreamName, const FOnDeleteFinishedStreamComplete& Delegate ) const
