Support showing number of views in enumerated replays

diff --git a/Engine/Source/Runtime/NetworkReplayStreaming/HttpNetworkReplayStreaming/Private/HttpNetworkReplayStreaming.cpp b/Engine/Source/Runtime/NetworkReplayStreaming/HttpNetworkReplayStreaming/Private/HttpNetworkReplayStreaming.cpp
--- a/Engine/Source/Runtime/NetworkReplayStreaming/HttpNetworkReplayStreaming/Private/HttpNetworkReplayStreaming.cpp
+++ b/Engine/Source/Runtime/NetworkReplayStreaming/HttpNetworkReplayStreaming/Private/HttpNetworkReplayStreaming.cpp
@@ -717,7 +717,7 @@ void FHttpNetworkReplayStreamer::HttpEnumerateSessionsFinished( FHttpRequestPtr
 			return;
 		}
 
-		const int NUM_TOKENS_PER_INFO = 4;
+		const int NUM_TOKENS_PER_INFO = 6;
 
 		if ( Tokens.Num() == 0 || ( Tokens.Num() % NUM_TOKENS_PER_INFO ) != 0 )
 		{
@@ -737,21 +737,31 @@ void FHttpNetworkReplayStreamer::HttpEnumerateSessionsFinished( FHttpRequestPtr
 			NewStream.SizeInBytes	= 0;
 			NewStream.Timestamp		= 0;
 
-			if ( i + 1 < Tokens.Num() )
+			if ( Tokens.IsValidIndex( i + 1 ) )
 			{
 				// Server returns milliseconds from January 1, 1970, 00:00:00 GMT
 				// We need to compensate for the fact that FDateTime starts at January 1, 0001 A.D. and is in 100 nanosecond resolution
 				NewStream.Timestamp = FDateTime( FCString::Atoi64( *Tokens[ i + 1 ] ) * 1000 * 10 + FDateTime( 1970, 1, 1 ).GetTicks() );
 			}
 
-			if ( i + 2 < Tokens.Num() )
+			if ( Tokens.IsValidIndex( i + 2 ) )
 			{
 				NewStream.SizeInBytes = FCString::Atoi( *Tokens[ i + 2 ] );
 			}
 
-			if ( i + 3 < Tokens.Num() )
+			if ( Tokens.IsValidIndex( i + 3 ) )
 			{
-				NewStream.bIsLive = Tokens[ i + 3 ].Contains( TEXT( "true" ) );
+				NewStream.LengthInMS = FCString::Atoi( *Tokens[ i + 3 ] );
+			}
+
+			if ( Tokens.IsValidIndex( i + 4 ) )
+			{
+				NewStream.NumViewers = FCString::Atoi( *Tokens[ i + 4 ] );
+			}
+
+			if ( Tokens.IsValidIndex( i + 5 ) )
+			{
+				NewStream.bIsLive = Tokens[ i + 5 ].Contains( TEXT( "true" ) );
 			}
 
 			Streams.Add( NewStream );
diff --git a/Engine/Source/Runtime/NetworkReplayStreaming/NetworkReplayStreaming/Public/NetworkReplayStreaming.h b/Engine/Source/Runtime/NetworkReplayStreaming/NetworkReplayStreaming/Public/NetworkReplayStreaming.h
--- a/Engine/Source/Runtime/NetworkReplayStreaming/NetworkReplayStreaming/Public/NetworkReplayStreaming.h
+++ b/Engine/Source/Runtime/NetworkReplayStreaming/NetworkReplayStreaming/Public/NetworkReplayStreaming.h
@@ -9,6 +9,8 @@
 /** Struct to store information about a stream, returned from search results. */
 struct FNetworkReplayStreamInfo
 {
+	FNetworkReplayStreamInfo() : SizeInBytes( 0 ), LengthInMS( 0 ), NumViewers( 0 ), bIsLive( false ) {}
+
 	/** The name of the stream */
 	FString Name;
 
@@ -17,6 +19,12 @@ struct FNetworkReplayStreamInfo
 
 	/** The size of the stream */
 	int64 SizeInBytes;
+	
+	/** The duration of the stream in MS */
+	int32 LengthInMS;
+
+	/** Number of viewers viewing this stream */
+	int32 NumViewers;
 
 	/** True if the stream is live and the game hasn't completed yet */
 	bool bIsLive;
