Make http replay logging less verbose

diff --git a/Engine/Source/Runtime/NetworkReplayStreaming/HttpNetworkReplayStreaming/Private/HttpNetworkReplayStreaming.cpp b/Engine/Source/Runtime/NetworkReplayStreaming/HttpNetworkReplayStreaming/Private/HttpNetworkReplayStreaming.cpp
--- a/Engine/Source/Runtime/NetworkReplayStreaming/HttpNetworkReplayStreaming/Private/HttpNetworkReplayStreaming.cpp
+++ b/Engine/Source/Runtime/NetworkReplayStreaming/HttpNetworkReplayStreaming/Private/HttpNetworkReplayStreaming.cpp
@@ -491,7 +491,7 @@ void FHttpNetworkReplayStreamer::HttpUploadFinished( FHttpRequestPtr , FHttpResp
 
 	if ( bSucceeded && HttpResponse->GetResponseCode() == EHttpResponseCodes::Ok )
 	{
-		UE_LOG( LogHttpReplay, Log, TEXT( "FHttpNetworkReplayStreamer::HttpUploadFinished." ) );
+		UE_LOG( LogHttpReplay, Verbose, TEXT( "FHttpNetworkReplayStreamer::HttpUploadFinished." ) );
 
 		if ( StreamerState == EStreamerState::StreamingUpFinal )
 		{
@@ -627,7 +627,7 @@ void FHttpNetworkReplayStreamer::HttpDownloadFinished( FHttpRequestPtr HttpReque
 				StreamFileCount++;
 			}
 
-			UE_LOG( LogHttpReplay, Log, TEXT( "FHttpNetworkReplayStreamer::HttpDownloadFinished. State: %s, Progress: %i / %i, DemoTime: %2.2f" ), *State, StreamFileCount, NumDownloadChunks, (float)DemoTimeInMS / 1000 );
+			UE_LOG( LogHttpReplay, Verbose, TEXT( "FHttpNetworkReplayStreamer::HttpDownloadFinished. State: %s, Progress: %i / %i, DemoTime: %2.2f" ), *State, StreamFileCount, NumDownloadChunks, (float)DemoTimeInMS / 1000 );
 		}
 		else
 		{
@@ -670,7 +670,7 @@ void FHttpNetworkReplayStreamer::HttpEnumerateSessionsFinished( FHttpRequestPtr
 
 	if ( bSucceeded && HttpResponse->GetResponseCode() == EHttpResponseCodes::Ok )
 	{		
-		UE_LOG( LogHttpReplay, Log, TEXT( "FHttpNetworkReplayStreamer::HttpEnumerateSessionsFinished." ) );
+		UE_LOG( LogHttpReplay, Verbose, TEXT( "FHttpNetworkReplayStreamer::HttpEnumerateSessionsFinished." ) );
 
 		TArray<FNetworkReplayStreamInfo> Streams;
 		FString StreamsString = HttpResponse->GetContentAsString();
