Make http replay streamer even less verbose

diff --git a/Engine/Source/Runtime/NetworkReplayStreaming/HttpNetworkReplayStreaming/Private/HttpNetworkReplayStreaming.cpp b/Engine/Source/Runtime/NetworkReplayStreaming/HttpNetworkReplayStreaming/Private/HttpNetworkReplayStreaming.cpp
--- a/Engine/Source/Runtime/NetworkReplayStreaming/HttpNetworkReplayStreaming/Private/HttpNetworkReplayStreaming.cpp
+++ b/Engine/Source/Runtime/NetworkReplayStreaming/HttpNetworkReplayStreaming/Private/HttpNetworkReplayStreaming.cpp
@@ -227,7 +227,7 @@ void FHttpNetworkReplayStreamer::FlushStream()
 	}
 
 	// Upload any new streamed data to the http server
-	UE_LOG( LogHttpReplay, Log, TEXT( "FHttpNetworkReplayStreamer::FlushStream. StreamFileCount: %i, Size: %i" ), StreamFileCount, StreamArchive.Buffer.Num() );
+	UE_LOG( LogHttpReplay, Verbose, TEXT( "FHttpNetworkReplayStreamer::FlushStream. StreamFileCount: %i, Size: %i" ), StreamFileCount, StreamArchive.Buffer.Num() );
 
 	// Create the Http request and add to pending request list
 	TSharedRef<class IHttpRequest> HttpRequest = FHttpModule::Get().CreateRequest();
