Merge: Add error handling to network replay interface

diff --git a/Engine/Source/Runtime/NetworkReplayStreaming/HttpNetworkReplayStreaming/Private/HttpNetworkReplayStreaming.cpp b/Engine/Source/Runtime/NetworkReplayStreaming/HttpNetworkReplayStreaming/Private/HttpNetworkReplayStreaming.cpp
--- a/Engine/Source/Runtime/NetworkReplayStreaming/HttpNetworkReplayStreaming/Private/HttpNetworkReplayStreaming.cpp
+++ b/Engine/Source/Runtime/NetworkReplayStreaming/HttpNetworkReplayStreaming/Private/HttpNetworkReplayStreaming.cpp
@@ -59,7 +59,8 @@ FHttpNetworkReplayStreamer::FHttpNetworkReplayStreamer() :
 	bStopStreamingCalled( false ), 
 	bStreamIsLive( false ),
 	NumDownloadChunks( 0 ),
-	DemoTimeInMS( 0 )
+	DemoTimeInMS( 0 ),
+	StreamerLastError( ENetworkReplayError::None )
 {
 	// Initialize the server URL
 	GConfig->GetString( TEXT( "HttpNetworkReplayStreaming" ), TEXT( "ServerURL" ), ServerURL, GEngineIni );
@@ -350,6 +351,17 @@ void FHttpNetworkReplayStreamer::RefreshViewer()
 	LastRefreshViewerTime = FPlatformTime::Seconds();
 }
 
+void FHttpNetworkReplayStreamer::SetLastError( const ENetworkReplayError::Type InLastError )
+{
+	StreamerState		= EStreamerState::Idle;
+	StreamerLastError	= InLastError;
+}
+
+ENetworkReplayError::Type FHttpNetworkReplayStreamer::GetLastError() const
+{
+	return StreamerLastError;
+}
+
 FArchive* FHttpNetworkReplayStreamer::GetHeaderArchive()
 {
 	return ( HeaderArchive.IsSaving() || HeaderArchive.Pos < HeaderArchive.Buffer.Num() ) ? &HeaderArchive : NULL;
@@ -429,7 +441,8 @@ void FHttpNetworkReplayStreamer::HttpStartUploadingFinished( FHttpRequestPtr Htt
 	}
 	else
 	{
-		UE_LOG( LogHttpReplay, Warning, TEXT( "FHttpNetworkReplayStreamer::HttpStartUploadingFinished. FAILED" ) );
+		UE_LOG( LogHttpReplay, Error, TEXT( "FHttpNetworkReplayStreamer::HttpStartUploadingFinished. FAILED" ) );
+		SetLastError( ENetworkReplayError::ServiceUnavailable );
 	}
 }
 
@@ -447,7 +460,8 @@ void FHttpNetworkReplayStreamer::HttpStopUploadingFinished( FHttpRequestPtr Http
 	}
 	else
 	{
-		UE_LOG( LogHttpReplay, Warning, TEXT( "FHttpNetworkReplayStreamer::HttpStopUploadingFinished. FAILED" ) );
+		UE_LOG( LogHttpReplay, Error, TEXT( "FHttpNetworkReplayStreamer::HttpStopUploadingFinished. FAILED" ) );
+		SetLastError( ENetworkReplayError::ServiceUnavailable );
 	}
 
 	StreamArchive.ArIsLoading = false;
@@ -474,11 +488,8 @@ void FHttpNetworkReplayStreamer::HttpHeaderUploadFinished( FHttpRequestPtr HttpR
 	}
 	else
 	{
-		UE_LOG( LogHttpReplay, Warning, TEXT( "FHttpNetworkReplayStreamer::HttpHeaderUploadFinished. FAILED" ) );
-
-		StreamerState = EStreamerState::Idle;
-
-		// FIXME: Notify demo driver with delegate
+		UE_LOG( LogHttpReplay, Error, TEXT( "FHttpNetworkReplayStreamer::HttpHeaderUploadFinished. FAILED" ) );
+		SetLastError( ENetworkReplayError::ServiceUnavailable );
 	}
 }
 
@@ -511,8 +522,8 @@ void FHttpNetworkReplayStreamer::HttpUploadFinished( FHttpRequestPtr , FHttpResp
 	}
 	else
 	{
-		UE_LOG( LogHttpReplay, Warning, TEXT( "FHttpNetworkReplayStreamer::HttpUploadFinished. FAILED" ) );
-		StreamerState = EStreamerState::Idle;
+		UE_LOG( LogHttpReplay, Error, TEXT( "FHttpNetworkReplayStreamer::HttpUploadFinished. FAILED" ) );
+		SetLastError( ENetworkReplayError::ServiceUnavailable );
 	}
 }
 
@@ -552,20 +563,19 @@ void FHttpNetworkReplayStreamer::HttpStartDownloadingFinished( FHttpRequestPtr H
 			// Reset delegate
 			StartStreamingDelegate = FOnStreamReadyDelegate();
 
-			StreamerState = EStreamerState::Idle;
+			SetLastError( ENetworkReplayError::ServiceUnavailable );
 		}
 	}
 	else
 	{
-		// FAIL
-		UE_LOG( LogHttpReplay, Warning, TEXT( "FHttpNetworkReplayStreamer::HttpStartDownloadingFinished. FAILED" ) );
+		UE_LOG( LogHttpReplay, Error, TEXT( "FHttpNetworkReplayStreamer::HttpStartDownloadingFinished. FAILED" ) );
 
 		StartStreamingDelegate.ExecuteIfBound( false, StreamArchive.IsSaving() );
 
 		// Reset delegate
 		StartStreamingDelegate = FOnStreamReadyDelegate();
 
-		StreamerState = EStreamerState::Idle;
+		SetLastError( ENetworkReplayError::ServiceUnavailable );
 	}
 }
 
@@ -589,12 +599,12 @@ void FHttpNetworkReplayStreamer::HttpDownloadHeaderFinished( FHttpRequestPtr Htt
 	else
 	{
 		// FAIL
-		UE_LOG( LogHttpReplay, Warning, TEXT( "FHttpNetworkReplayStreamer::HttpDownloadHeaderFinished. FAILED." ) );
+		UE_LOG( LogHttpReplay, Error, TEXT( "FHttpNetworkReplayStreamer::HttpDownloadHeaderFinished. FAILED." ) );
 
 		StreamArchive.Buffer.Empty();
 		StartStreamingDelegate.ExecuteIfBound( false, StreamArchive.IsSaving() );
 
-		StreamerState = EStreamerState::Idle;
+		SetLastError( ENetworkReplayError::ServiceUnavailable );
 	}
 
 	// Reset delegate
@@ -631,20 +641,16 @@ void FHttpNetworkReplayStreamer::HttpDownloadFinished( FHttpRequestPtr HttpReque
 		}
 		else
 		{
-			UE_LOG( LogHttpReplay, Warning, TEXT( "FHttpNetworkReplayStreamer::HttpDownloadFinished. FAILED." ) );
+			UE_LOG( LogHttpReplay, Error, TEXT( "FHttpNetworkReplayStreamer::HttpDownloadFinished. FAILED." ) );
 			StreamArchive.Buffer.Empty();
-			StreamerState = EStreamerState::Idle;
-
-			// FIXME: Report error to engine
+			SetLastError( ENetworkReplayError::ServiceUnavailable );
 		}		
 	}
 	else
 	{
-		UE_LOG( LogHttpReplay, Warning, TEXT( "FHttpNetworkReplayStreamer::HttpDownloadFinished. FAILED." ) );
+		UE_LOG( LogHttpReplay, Error, TEXT( "FHttpNetworkReplayStreamer::HttpDownloadFinished. FAILED." ) );
 		StreamArchive.Buffer.Empty();
-		StreamerState = EStreamerState::Idle;
-
-		// FIXME: Report error to engine
+		SetLastError( ENetworkReplayError::ServiceUnavailable );
 	}
 }
 
@@ -657,8 +663,8 @@ void FHttpNetworkReplayStreamer::HttpRefreshViewerFinished( FHttpRequestPtr Http
 
 	if ( !bSucceeded || HttpResponse->GetResponseCode() != EHttpResponseCodes::Ok )
 	{
-		// FAIL
-		UE_LOG( LogHttpReplay, Warning, TEXT( "FHttpNetworkReplayStreamer::HttpRefreshViewerFinished. FAILED." ) );
+		UE_LOG( LogHttpReplay, Error, TEXT( "FHttpNetworkReplayStreamer::HttpRefreshViewerFinished. FAILED." ) );
+		SetLastError( ENetworkReplayError::ServiceUnavailable );
 	}
 }
 
diff --git a/Engine/Source/Runtime/NetworkReplayStreaming/HttpNetworkReplayStreaming/Private/HttpNetworkReplayStreaming.h b/Engine/Source/Runtime/NetworkReplayStreaming/HttpNetworkReplayStreaming/Private/HttpNetworkReplayStreaming.h
--- a/Engine/Source/Runtime/NetworkReplayStreaming/HttpNetworkReplayStreaming/Private/HttpNetworkReplayStreaming.h
+++ b/Engine/Source/Runtime/NetworkReplayStreaming/HttpNetworkReplayStreaming/Private/HttpNetworkReplayStreaming.h
@@ -42,6 +42,7 @@ public:
 	virtual bool		IsLive( const FString& StreamName ) const override;
 	virtual void		DeleteFinishedStream( const FString& StreamName, const FOnDeleteFinishedStreamComplete& Delegate ) const override;
 	virtual void		EnumerateStreams( const FString& VersionString, const FOnEnumerateStreamsComplete& Delegate ) override;
+	virtual ENetworkReplayError::Type GetLastError() const override;
 
 	/** FHttpNetworkReplayStreamer */
 	void UploadHeader();
@@ -49,6 +50,7 @@ public:
 	void DownloadHeader();
 	void DownloadNextChunk();
 	void RefreshViewer();
+	void SetLastError( const ENetworkReplayError::Type InLastError );
 
 	/** Delegates */
 	void HttpStartDownloadingFinished( FHttpRequestPtr HttpRequest, FHttpResponsePtr HttpResponse, bool bSucceeded );
@@ -108,6 +110,8 @@ public:
 	uint32					DemoTimeInMS;
 	FString					ViewerName;
 
+	ENetworkReplayError::Type		StreamerLastError;
+
 	FOnStreamReadyDelegate			StartStreamingDelegate;		// Delegate passed in to StartStreaming
 	FOnEnumerateStreamsComplete		EnumerateStreamsDelegate;
 };
diff --git a/Engine/Source/Runtime/NetworkReplayStreaming/NetworkReplayStreaming/Public/NetworkReplayStreaming.h b/Engine/Source/Runtime/NetworkReplayStreaming/NetworkReplayStreaming/Public/NetworkReplayStreaming.h
--- a/Engine/Source/Runtime/NetworkReplayStreaming/NetworkReplayStreaming/Public/NetworkReplayStreaming.h
+++ b/Engine/Source/Runtime/NetworkReplayStreaming/NetworkReplayStreaming/Public/NetworkReplayStreaming.h
@@ -30,6 +30,30 @@ struct FNetworkReplayStreamInfo
 	bool bIsLive;
 };
 
+namespace ENetworkReplayError
+{
+	enum Type
+	{
+		/** There are currently no issues */
+		None,
+
+		/** The backend service supplying the stream is unavailable, or connection interrupted */
+		ServiceUnavailable,
+	};
+
+	inline const TCHAR* ToString( const ENetworkReplayError::Type FailureType )
+	{
+		switch ( FailureType )
+		{
+			case None:
+				return TEXT( "None" );
+			case ServiceUnavailable:
+				return TEXT( "ServiceUnavailable" );
+		}
+		return TEXT( "Unknown ENetworkReplayError error" );
+	}
+}
+
 DECLARE_MULTICAST_DELEGATE_TwoParams( FOnStreamReady, bool, bool );
 typedef FOnStreamReady::FDelegate FOnStreamReadyDelegate;
 
@@ -79,6 +103,9 @@ public:
 	 * @param Delegate A delegate that will be executed if bound when the list of streams is available
 	 */
 	virtual void EnumerateStreams( const FString& VersionString, const FOnEnumerateStreamsComplete& Delegate ) = 0;
+
+	/** Returns the last error that occurred while streaming replays */
+	virtual ENetworkReplayError::Type GetLastError() const = 0;
 };
 
 /** Replay streamer factory */
diff --git a/Engine/Source/Runtime/NetworkReplayStreaming/NullNetworkReplayStreaming/Public/NullNetworkReplayStreaming.h b/Engine/Source/Runtime/NetworkReplayStreaming/NullNetworkReplayStreaming/Public/NullNetworkReplayStreaming.h
--- a/Engine/Source/Runtime/NetworkReplayStreaming/NullNetworkReplayStreaming/Public/NullNetworkReplayStreaming.h
+++ b/Engine/Source/Runtime/NetworkReplayStreaming/NullNetworkReplayStreaming/Public/NullNetworkReplayStreaming.h
@@ -27,6 +27,7 @@ public:
 	virtual bool IsLive( const FString& StreamName ) const override;
 	virtual void DeleteFinishedStream( const FString& StreamName, const FOnDeleteFinishedStreamComplete& Delegate) const override;
 	virtual void EnumerateStreams( const FString& VersionString, const FOnEnumerateStreamsComplete& Delegate ) override;
+	virtual ENetworkReplayError::Type GetLastError() const override { return ENetworkReplayError::None; }
 
 private:
 	/** Handle to the archive that will read/write network packets */
