Merge:  Initial support for friendly names in replays

diff --git a/Engine/Source/Runtime/NetworkReplayStreaming/HttpNetworkReplayStreaming/Private/HttpNetworkReplayStreaming.cpp b/Engine/Source/Runtime/NetworkReplayStreaming/HttpNetworkReplayStreaming/Private/HttpNetworkReplayStreaming.cpp
--- a/Engine/Source/Runtime/NetworkReplayStreaming/HttpNetworkReplayStreaming/Private/HttpNetworkReplayStreaming.cpp
+++ b/Engine/Source/Runtime/NetworkReplayStreaming/HttpNetworkReplayStreaming/Private/HttpNetworkReplayStreaming.cpp
@@ -100,14 +100,6 @@ void FHttpNetworkReplayStreamer::StartStreaming( const FString& StreamName, bool
 
 	LastChunkTime = FPlatformTime::Seconds();
 
-	const bool bOverrideRecordingSession = false;
-
-	// Override session name if requested
-	if ( StreamArchive.ArIsLoading || bOverrideRecordingSession )
-	{
-		SessionName = StreamName;
-	}
-
 	// Create the Http request and add to pending request list
 	TSharedRef<class IHttpRequest> HttpRequest = FHttpModule::Get().CreateRequest();
 
@@ -115,6 +107,8 @@ void FHttpNetworkReplayStreamer::StartStreaming( const FString& StreamName, bool
 
 	if ( StreamArchive.ArIsLoading )
 	{
+		SessionName = StreamName;
+
 		// Notify the http server that we want to start downloading a replay
 		HttpRequest->SetURL( FString::Printf( TEXT( "%sstartdownloading?Version=%s&Session=%s" ), *ServerURL, *SessionVersion, *SessionName ) );
 		HttpRequest->SetVerb( TEXT( "POST" ) );
@@ -128,15 +122,8 @@ void FHttpNetworkReplayStreamer::StartStreaming( const FString& StreamName, bool
 	}
 	else
 	{
-		// Notify the http server that we want to start upload a replay
-		if ( !SessionName.IsEmpty() )
-		{
-			HttpRequest->SetURL( FString::Printf( TEXT( "%sstartuploading?Version=%s&Session=%s" ), *ServerURL, *SessionVersion, *SessionName ) );
-		}
-		else
-		{
-			HttpRequest->SetURL( FString::Printf( TEXT( "%sstartuploading?Version=%s" ), *ServerURL, *SessionVersion ) );
-		}
+		// Notify the http server that we want to start uploading a replay
+		HttpRequest->SetURL( FString::Printf( TEXT( "%sstartuploading?Version=%s&Friendly=%s" ), *ServerURL, *SessionVersion, *StreamName ) );
 
 		HttpRequest->SetVerb( TEXT( "POST" ) );
 
