Use JSON for replay enumeration

diff --git a/Engine/Source/Runtime/NetworkReplayStreaming/HttpNetworkReplayStreaming/HttpNetworkReplayStreaming.Build.cs b/Engine/Source/Runtime/NetworkReplayStreaming/HttpNetworkReplayStreaming/HttpNetworkReplayStreaming.Build.cs
--- a/Engine/Source/Runtime/NetworkReplayStreaming/HttpNetworkReplayStreaming/HttpNetworkReplayStreaming.Build.cs
+++ b/Engine/Source/Runtime/NetworkReplayStreaming/HttpNetworkReplayStreaming/HttpNetworkReplayStreaming.Build.cs
@@ -8,6 +8,8 @@ namespace UnrealBuildTool.Rules
         {
 			PrivateIncludePaths.Add( "Runtime/NetworkReplayStreaming/HttpNetworkReplayStreaming/Private" );
 
+			PrivateIncludePathModuleNames.Add( "OnlineSubsystem" );
+
 			PrivateDependencyModuleNames.AddRange(
                 new string[]
 				{
@@ -15,8 +17,9 @@ namespace UnrealBuildTool.Rules
 					"CoreUObject",
                     "Engine",
                     "HTTP",
-					"NetworkReplayStreaming"
-				});
+					"NetworkReplayStreaming",
+					"Json",
+				} );
         }
     }
 }
diff --git a/Engine/Source/Runtime/NetworkReplayStreaming/HttpNetworkReplayStreaming/Private/HttpNetworkReplayStreaming.cpp b/Engine/Source/Runtime/NetworkReplayStreaming/HttpNetworkReplayStreaming/Private/HttpNetworkReplayStreaming.cpp
--- a/Engine/Source/Runtime/NetworkReplayStreaming/HttpNetworkReplayStreaming/Private/HttpNetworkReplayStreaming.cpp
+++ b/Engine/Source/Runtime/NetworkReplayStreaming/HttpNetworkReplayStreaming/Private/HttpNetworkReplayStreaming.cpp
@@ -1,9 +1,53 @@
 // Copyright 1998-2015 Epic Games, Inc. All Rights Reserved.
 
 #include "HttpNetworkReplayStreaming.h"
+#include "OnlineJsonSerializer.h"
 
 DEFINE_LOG_CATEGORY_STATIC( LogHttpReplay, Log, All );
 
+class FNetworkReplayListItem : public FOnlineJsonSerializable
+{
+public:
+	FNetworkReplayListItem() {}
+	virtual ~FNetworkReplayListItem() {}
+
+	FString		VersionName;
+	FString		SessionName;
+	FString		FriendlyName;
+	FDateTime	Timestamp;
+	int32		SizeInBytes;
+	int32		DemoTimeInMs;
+	int32		NumViewers;
+	bool		bIsLive;
+
+	// FOnlineJsonSerializable
+	BEGIN_ONLINE_JSON_SERIALIZER
+		ONLINE_JSON_SERIALIZE( "VersionName",	VersionName );
+		ONLINE_JSON_SERIALIZE( "SessionName",	SessionName );
+		ONLINE_JSON_SERIALIZE( "FriendlyName",	FriendlyName );
+		ONLINE_JSON_SERIALIZE( "Timestamp",		Timestamp );
+		ONLINE_JSON_SERIALIZE( "SizeInBytes",	SizeInBytes );
+		ONLINE_JSON_SERIALIZE( "DemoTimeInMs",	DemoTimeInMs );
+		ONLINE_JSON_SERIALIZE( "NumViewers",	NumViewers );
+		ONLINE_JSON_SERIALIZE( "bIsLive",		bIsLive );
+	END_ONLINE_JSON_SERIALIZER
+};
+
+class FNetworkReplayList : public FOnlineJsonSerializable
+{
+public:
+	FNetworkReplayList()
+	{}
+	virtual ~FNetworkReplayList() {}
+
+	TArray< FNetworkReplayListItem > Replays;
+
+	// FOnlineJsonSerializable
+	BEGIN_ONLINE_JSON_SERIALIZER
+		ONLINE_JSON_SERIALIZE_ARRAY_SERIALIZABLE( "replays", Replays, FNetworkReplayListItem );
+	END_ONLINE_JSON_SERIALIZER
+};
+
 void HttpStreamFArchive::Serialize( void* V, int64 Length ) 
 {
 	if ( IsLoading() )
@@ -666,96 +710,27 @@ void FHttpNetworkReplayStreamer::HttpEnumerateSessionsFinished( FHttpRequestPtr
 		UE_LOG( LogHttpReplay, Verbose, TEXT( "FHttpNetworkReplayStreamer::HttpEnumerateSessionsFinished." ) );
 
 		TArray<FNetworkReplayStreamInfo> Streams;
-		FString StreamsString = HttpResponse->GetContentAsString();
-
-		int32 Index = INDEX_NONE;
+		FString JsonString = HttpResponse->GetContentAsString();
 
-		TArray< FString > Tokens;
+		FNetworkReplayList ReplayList;
 
-		// Parse the string as { token 1, token ..., token n }
-		// This isn't perfect, we should convert to JSON when the dust settles
-		if ( StreamsString.FindChar( '{', Index ) )
+		if ( !ReplayList.FromJson( JsonString ) )
 		{
-			StreamsString = StreamsString.RightChop( Index + 1 );
-
-			while ( StreamsString.FindChar( ',', Index ) )
-			{
-				Tokens.Add( StreamsString.Left( Index ) );
-				StreamsString = StreamsString.RightChop( Index + 1 );
-			}
-
-			if ( StreamsString.FindChar( '}', Index ) )
-			{
-				Tokens.Add( StreamsString.Left( Index ) );
-			}
-			else
-			{
-				UE_LOG( LogHttpReplay, Warning, TEXT( "FHttpNetworkReplayStreamer::HttpEnumerateSessionsFinished. '}' not found." ) );
-				EnumerateStreamsDelegate.ExecuteIfBound( TArray<FNetworkReplayStreamInfo>() );		// FIXME: Notify failure here
-				EnumerateStreamsDelegate = FOnEnumerateStreamsComplete();
-				return;
-			}
-
-			if ( Tokens.Num() > 0 )
-			{
-				Tokens[ Tokens.Num() - 1 ].RemoveFromStart( TEXT( " " ) );
-				Tokens[ Tokens.Num() - 1 ].RemoveFromEnd( TEXT( " " ) );
-			}
-		}
-		else
-		{
-			UE_LOG( LogHttpReplay, Warning, TEXT( "FHttpNetworkReplayStreamer::HttpEnumerateSessionsFinished. '{' not found." ) );
+			UE_LOG( LogHttpReplay, Warning, TEXT( "FHttpNetworkReplayStreamer::HttpEnumerateSessionsFinished. FromJson FAILED" ) );
 			EnumerateStreamsDelegate.ExecuteIfBound( TArray<FNetworkReplayStreamInfo>() );		// FIXME: Notify failure here
-			EnumerateStreamsDelegate = FOnEnumerateStreamsComplete();
 			return;
 		}
 
-		const int NUM_TOKENS_PER_INFO = 6;
-
-		if ( Tokens.Num() == 0 || ( Tokens.Num() % NUM_TOKENS_PER_INFO ) != 0 )
-		{
-			UE_LOG( LogHttpReplay, Warning, TEXT( "FHttpNetworkReplayStreamer::HttpEnumerateSessionsFinished. Invalid number of tokens: %i" ), Tokens.Num() );
-			EnumerateStreamsDelegate.ExecuteIfBound( TArray<FNetworkReplayStreamInfo>() );		// FIXME: Notify failure here
-			EnumerateStreamsDelegate = FOnEnumerateStreamsComplete();
-			return;
-		}
-
-		// Convert tokens to individual FNetworkReplayStreamInfo's
-		for ( int i = 0; i < Tokens.Num(); i += NUM_TOKENS_PER_INFO )
+		for ( int32 i = 0; i < ReplayList.Replays.Num(); i++ )
 		{
 			FNetworkReplayStreamInfo NewStream;
 
-			NewStream.Name			= Tokens[i];
-			NewStream.bIsLive		= false;
-			NewStream.SizeInBytes	= 0;
-			NewStream.Timestamp		= 0;
-
-			if ( Tokens.IsValidIndex( i + 1 ) )
-			{
-				// Server returns milliseconds from January 1, 1970, 00:00:00 GMT
-				// We need to compensate for the fact that FDateTime starts at January 1, 0001 A.D. and is in 100 nanosecond resolution
-				NewStream.Timestamp = FDateTime( FCString::Atoi64( *Tokens[ i + 1 ] ) * 1000 * 10 + FDateTime( 1970, 1, 1 ).GetTicks() );
-			}
-
-			if ( Tokens.IsValidIndex( i + 2 ) )
-			{
-				NewStream.SizeInBytes = FCString::Atoi( *Tokens[ i + 2 ] );
-			}
-
-			if ( Tokens.IsValidIndex( i + 3 ) )
-			{
-				NewStream.LengthInMS = FCString::Atoi( *Tokens[ i + 3 ] );
-			}
-
-			if ( Tokens.IsValidIndex( i + 4 ) )
-			{
-				NewStream.NumViewers = FCString::Atoi( *Tokens[ i + 4 ] );
-			}
-
-			if ( Tokens.IsValidIndex( i + 5 ) )
-			{
-				NewStream.bIsLive = Tokens[ i + 5 ].Contains( TEXT( "true" ) );
-			}
+			NewStream.Name			= ReplayList.Replays[i].SessionName;
+			NewStream.Timestamp		= ReplayList.Replays[i].Timestamp;
+			NewStream.SizeInBytes	= ReplayList.Replays[i].SizeInBytes;
+			NewStream.LengthInMS	= ReplayList.Replays[i].DemoTimeInMs;
+			NewStream.NumViewers	= ReplayList.Replays[i].NumViewers;
+			NewStream.bIsLive		= ReplayList.Replays[i].bIsLive;;
 
 			Streams.Add( NewStream );
 		}
