Add friendly name to FNetworkReplayStreamInfo

diff --git a/Engine/Source/Runtime/NetworkReplayStreaming/HttpNetworkReplayStreaming/Private/HttpNetworkReplayStreaming.cpp b/Engine/Source/Runtime/NetworkReplayStreaming/HttpNetworkReplayStreaming/Private/HttpNetworkReplayStreaming.cpp
--- a/Engine/Source/Runtime/NetworkReplayStreaming/HttpNetworkReplayStreaming/Private/HttpNetworkReplayStreaming.cpp
+++ b/Engine/Source/Runtime/NetworkReplayStreaming/HttpNetworkReplayStreaming/Private/HttpNetworkReplayStreaming.cpp
@@ -726,6 +726,7 @@ void FHttpNetworkReplayStreamer::HttpEnumerateSessionsFinished( FHttpRequestPtr
 			FNetworkReplayStreamInfo NewStream;
 
 			NewStream.Name			= ReplayList.Replays[i].SessionName;
+			NewStream.FriendlyName	= ReplayList.Replays[i].FriendlyName;
 			NewStream.Timestamp		= ReplayList.Replays[i].Timestamp;
 			NewStream.SizeInBytes	= ReplayList.Replays[i].SizeInBytes;
 			NewStream.LengthInMS	= ReplayList.Replays[i].DemoTimeInMs;
diff --git a/Engine/Source/Runtime/NetworkReplayStreaming/NetworkReplayStreaming/Public/NetworkReplayStreaming.h b/Engine/Source/Runtime/NetworkReplayStreaming/NetworkReplayStreaming/Public/NetworkReplayStreaming.h
--- a/Engine/Source/Runtime/NetworkReplayStreaming/NetworkReplayStreaming/Public/NetworkReplayStreaming.h
+++ b/Engine/Source/Runtime/NetworkReplayStreaming/NetworkReplayStreaming/Public/NetworkReplayStreaming.h
@@ -11,9 +11,12 @@ struct FNetworkReplayStreamInfo
 {
 	FNetworkReplayStreamInfo() : SizeInBytes( 0 ), LengthInMS( 0 ), NumViewers( 0 ), bIsLive( false ) {}
 
-	/** The name of the stream */
+	/** The name of the stream (generally this is auto generated, refer to friendly name for UI) */
 	FString Name;
 
+	/** The UI friendly name of the stream */
+	FString FriendlyName;
+
 	/** The date and time the stream was recorded */
 	FDateTime Timestamp;
 
