Give spectate viewer refreshing priority + tweak update rates

diff --git a/Engine/Source/Runtime/NetworkReplayStreaming/HttpNetworkReplayStreaming/Private/HttpNetworkReplayStreaming.cpp b/Engine/Source/Runtime/NetworkReplayStreaming/HttpNetworkReplayStreaming/Private/HttpNetworkReplayStreaming.cpp
--- a/Engine/Source/Runtime/NetworkReplayStreaming/HttpNetworkReplayStreaming/Private/HttpNetworkReplayStreaming.cpp
+++ b/Engine/Source/Runtime/NetworkReplayStreaming/HttpNetworkReplayStreaming/Private/HttpNetworkReplayStreaming.cpp
@@ -350,7 +350,7 @@ void FHttpNetworkReplayStreamer::RefreshViewer()
 	}
 
 	// Determine if it's time to download the next chunk
-	const double REFRESH_VIEWER_IN_SECONDS = 5;
+	const double REFRESH_VIEWER_IN_SECONDS = 10;
 
 	if ( StreamerState != EStreamerState::StreamingDownFinal && FPlatformTime::Seconds() - LastRefreshViewerTime < REFRESH_VIEWER_IN_SECONDS )
 	{
@@ -800,8 +800,9 @@ void FHttpNetworkReplayStreamer::Tick( float DeltaTime )
 	}
 	else if ( StreamerState == EStreamerState::StreamingDown )
 	{
-		DownloadNextChunk();
+		// Give viewer refreshing priority since it's quick
 		RefreshViewer();
+		DownloadNextChunk();
 	}
 	else if ( StreamerState == EStreamerState::StreamingDownFinal )
 	{
