Merge UE4-Orion-to-UE4: Merge latest changes made in Orion branch.

diff --git a/Engine/Source/Runtime/NetworkReplayStreaming/HttpNetworkReplayStreaming/Private/HttpNetworkReplayStreaming.cpp b/Engine/Source/Runtime/NetworkReplayStreaming/HttpNetworkReplayStreaming/Private/HttpNetworkReplayStreaming.cpp
--- a/Engine/Source/Runtime/NetworkReplayStreaming/HttpNetworkReplayStreaming/Private/HttpNetworkReplayStreaming.cpp
+++ b/Engine/Source/Runtime/NetworkReplayStreaming/HttpNetworkReplayStreaming/Private/HttpNetworkReplayStreaming.cpp
@@ -96,7 +96,7 @@ void HttpStreamFArchive::Seek( int64 InPos )
 
 bool HttpStreamFArchive::AtEnd() 
 {
-	return bAtEndOfReplay;
+	return Pos >= Buffer.Num() && bAtEndOfReplay;
 }
 
 FHttpNetworkReplayStreamer::FHttpNetworkReplayStreamer() : 
@@ -782,7 +782,7 @@ void FHttpNetworkReplayStreamer::Tick( const float DeltaTime )
 		return;
 	}
 
-	if ( StreamFileCount >= NumDownloadChunks && bStreamIsLive )
+	if ( StreamFileCount >= NumDownloadChunks && !bStreamIsLive )
 	{
 		StreamArchive.bAtEndOfReplay = true;
 	}
