Convert http replay streamer to use http task queue, which greatly simplifies, and makes much more robust.

diff --git a/Engine/Source/Runtime/NetworkReplayStreaming/HttpNetworkReplayStreaming/Private/HttpNetworkReplayStreaming.cpp b/Engine/Source/Runtime/NetworkReplayStreaming/HttpNetworkReplayStreaming/Private/HttpNetworkReplayStreaming.cpp
--- a/Engine/Source/Runtime/NetworkReplayStreaming/HttpNetworkReplayStreaming/Private/HttpNetworkReplayStreaming.cpp
+++ b/Engine/Source/Runtime/NetworkReplayStreaming/HttpNetworkReplayStreaming/Private/HttpNetworkReplayStreaming.cpp
@@ -48,7 +48,7 @@ public:
 	END_ONLINE_JSON_SERIALIZER
 };
 
-void HttpStreamFArchive::Serialize( void* V, int64 Length ) 
+void FHttpStreamFArchive::Serialize( void* V, int64 Length ) 
 {
 	if ( IsLoading() )
 	{
@@ -79,24 +79,24 @@ void HttpStreamFArchive::Serialize( void* V, int64 Length )
 	}
 }
 
-int64 HttpStreamFArchive::Tell() 
+int64 FHttpStreamFArchive::Tell() 
 {
 	return Pos;
 }
 
-int64 HttpStreamFArchive::TotalSize()
+int64 FHttpStreamFArchive::TotalSize()
 {
 	return Buffer.Num();
 }
 
-void HttpStreamFArchive::Seek( int64 InPos ) 
+void FHttpStreamFArchive::Seek( int64 InPos ) 
 {
 	check( InPos < Buffer.Num() );
 
 	Pos = InPos;
 }
 
-bool HttpStreamFArchive::AtEnd() 
+bool FHttpStreamFArchive::AtEnd() 
 {
 	return Pos >= Buffer.Num() && bAtEndOfReplay;
 }
@@ -106,12 +106,11 @@ FHttpNetworkReplayStreamer::FHttpNetworkReplayStreamer() :
 	LastChunkTime( 0 ), 
 	LastRefreshViewerTime( 0 ),
 	StreamerState( EStreamerState::Idle ), 
-	HttpState( EHttptate::Idle ), 
 	bStopStreamingCalled( false ), 
+	bNeedToUploadHeader( false ),
 	bStreamIsLive( false ),
 	NumDownloadChunks( 0 ),
 	TotalDemoTimeInMS( 0 ),
-	FlushCheckpointTime( 0 ),
 	HighPriorityEndTime( 0 ),
 	StreamerLastError( ENetworkReplayError::None )
 {
@@ -171,11 +170,13 @@ void FHttpNetworkReplayStreamer::StartStreaming( const FString& StreamName, bool
 		HttpRequest->SetVerb( TEXT( "POST" ) );
 
 		HttpRequest->OnProcessRequestComplete().BindRaw( this, &FHttpNetworkReplayStreamer::HttpStartDownloadingFinished );
+	
+		// We are streaming down
+		StreamerState = EStreamerState::StreamingDown;
 
-		HttpState = EHttptate::StartDownloading;
-		
-		// Set the next streamer state to download the header
-		StreamerState = EStreamerState::NeedToDownloadHeader;
+		AddRequestToQueue( EQueuedHttpRequestType::StartDownloading, HttpRequest );
+
+		DownloadHeader();
 	}
 	else
 	{
@@ -188,27 +189,45 @@ void FHttpNetworkReplayStreamer::StartStreaming( const FString& StreamName, bool
 
 		SessionName.Empty();
 
-		HttpState = EHttptate::StartUploading;
+		// We are streaming up
+		StreamerState = EStreamerState::StreamingUp;
+
+		AddRequestToQueue( EQueuedHttpRequestType::StartUploading, HttpRequest );
 
-		// Set the next streamer state to upload the header
-		StreamerState = EStreamerState::NeedToUploadHeader;
+		// We can't upload the header until we have the session name, which depends on the StartUploading task above to finish
+		bNeedToUploadHeader = true;
 	}
-	
-	HttpRequest->ProcessRequest();
 }
 
-void FHttpNetworkReplayStreamer::StopStreaming()
+void FHttpNetworkReplayStreamer::AddRequestToQueue( const EQueuedHttpRequestType Type, TSharedPtr< class IHttpRequest >	Request )
 {
-	check( bStopStreamingCalled == false );
+	UE_LOG( LogHttpReplay, Verbose, TEXT( "FHttpNetworkReplayStreamer::AddRequestToQueue. Type: %i" ), (int)Type );
+
+	QueuedHttpRequests.Enqueue( TSharedPtr< FQueuedHttpRequest >( new FQueuedHttpRequest( Type, Request ) ) );
+}
 
-	if ( StreamerState == EStreamerState::Idle )
+void FHttpNetworkReplayStreamer::StopStreaming()
+{
+	if ( !IsStreaming() )
 	{
+		check( bStopStreamingCalled == false );
 		return;
 	}
 
-	// We need to queued this operation up, and finish it once all other operations are done
-	// (wait for header to upload, and any current stream to finish uploading, then we can flush what's left in StreamArchive)
 	bStopStreamingCalled = true;
+
+	if ( StreamerState == EStreamerState::StreamingDown )
+	{
+		// Refresh one last time, pasing in bFinal == true, this will notify the server we're done (remove us from viewers list)
+		RefreshViewer( true );
+	}
+	else if ( StreamerState == EStreamerState::StreamingUp )
+	{
+		// Flush any final pending stream
+		FlushStream();
+		// Send one last http request to stop uploading		
+		StopUploading();
+	}
 }
 
 void FHttpNetworkReplayStreamer::UploadHeader()
@@ -221,13 +240,7 @@ void FHttpNetworkReplayStreamer::UploadHeader()
 
 	if ( HeaderArchive.Buffer.Num() == 0 )
 	{
-		// We haven't serialized the header yet, don't do anything yet
-		return;
-	}
-
-	if ( IsHttpBusy() )
-	{
-		// If we're currently busy we can't flush right now
+		// Header wasn't serialized
 		return;
 	}
 
@@ -241,8 +254,6 @@ void FHttpNetworkReplayStreamer::UploadHeader()
 
 	HttpRequest->OnProcessRequestComplete().BindRaw( this, &FHttpNetworkReplayStreamer::HttpHeaderUploadFinished );
 
-	HttpState = EHttptate::UploadingHeader;
-
 	HttpRequest->SetURL( FString::Printf( TEXT( "%supload?Version=%s&Session=%s&NumChunks=%i&Time=%i&Filename=replay.header" ), *ServerURL, *SessionVersion, *SessionName, StreamFileCount, TotalDemoTimeInMS ) );
 	HttpRequest->SetVerb( TEXT( "POST" ) );
 	HttpRequest->SetHeader( TEXT( "Content-Type" ), TEXT( "application/octet-stream" ) );
@@ -252,7 +263,7 @@ void FHttpNetworkReplayStreamer::UploadHeader()
 	HeaderArchive.Buffer.Empty();
 	HeaderArchive.Pos = 0;
 
-	HttpRequest->ProcessRequest();
+	AddRequestToQueue( EQueuedHttpRequestType::UploadingHeader, HttpRequest );
 
 	LastChunkTime = FPlatformTime::Seconds();
 }
@@ -271,12 +282,6 @@ void FHttpNetworkReplayStreamer::FlushStream()
 		return;
 	}
 
-	if ( IsHttpBusy() )
-	{
-		// If we're currently busy we can't flush right now
-		return;
-	}
-
 	// Upload any new streamed data to the http server
 	UE_LOG( LogHttpReplay, Verbose, TEXT( "FHttpNetworkReplayStreamer::FlushStream. StreamFileCount: %i, Size: %i" ), StreamFileCount, StreamArchive.Buffer.Num() );
 
@@ -285,8 +290,6 @@ void FHttpNetworkReplayStreamer::FlushStream()
 
 	HttpRequest->OnProcessRequestComplete().BindRaw( this, &FHttpNetworkReplayStreamer::HttpUploadFinished );
 
-	HttpState = EHttptate::UploadingStream;
-
 	HttpRequest->SetURL( FString::Printf( TEXT( "%supload?Version=%s&Session=%s&NumChunks=%i&Time=%i&Filename=stream.%i" ), *ServerURL, *SessionVersion, *SessionName, StreamFileCount + 1, TotalDemoTimeInMS, StreamFileCount ) );
 	HttpRequest->SetVerb( TEXT( "POST" ) );
 	HttpRequest->SetHeader( TEXT( "Content-Type" ), TEXT( "application/octet-stream" ) );
@@ -297,25 +300,38 @@ void FHttpNetworkReplayStreamer::FlushStream()
 
 	StreamFileCount++;
 
-	HttpRequest->ProcessRequest();
+	AddRequestToQueue( EQueuedHttpRequestType::UploadingStream, HttpRequest );
 
 	LastChunkTime = FPlatformTime::Seconds();
 }
 
+void FHttpNetworkReplayStreamer::StopUploading()
+{
+	// Create the Http request and add to pending request list
+	TSharedRef<class IHttpRequest> HttpRequest = FHttpModule::Get().CreateRequest();
+
+	HttpRequest->OnProcessRequestComplete().BindRaw( this, &FHttpNetworkReplayStreamer::HttpStopUploadingFinished );
+
+	HttpRequest->SetURL( FString::Printf( TEXT( "%sstopuploading?Version=%s&Session=%s&NumChunks=%i&Time=%i" ), *ServerURL, *SessionVersion, *SessionName, StreamFileCount, TotalDemoTimeInMS ) );
+	HttpRequest->SetVerb( TEXT( "POST" ) );
+	HttpRequest->SetHeader( TEXT( "Content-Type" ), TEXT( "application/octet-stream" ) );
+
+	AddRequestToQueue( EQueuedHttpRequestType::StopUploading, HttpRequest );
+};
+
 void FHttpNetworkReplayStreamer::FlushCheckpoint( const uint32 TimeInMS )
 {
 	check( CheckpointArchive.Buffer.Num() > 0 );
-	FlushCheckpointTime = TimeInMS;
 	
 	// Flush any existing stream, we need checkpoints to line up with the next chunk
-	check( !IsHttpBusy() || StreamArchive.Buffer.Num() == 0 );
-
 	FlushStream();
+	// Flush the checkpoint
+	FlushCheckpointInternal( TimeInMS );
 }
 
 void FHttpNetworkReplayStreamer::GotoCheckpoint( const uint32 TimeInMS, const FOnCheckpointReadyDelegate& Delegate )
 {
-	if ( IsHttpBusy() )
+	if ( GotoCheckpointDelegate.IsBound() )
 	{
 		return;
 	}
@@ -323,29 +339,23 @@ void FHttpNetworkReplayStreamer::GotoCheckpoint( const uint32 TimeInMS, const FO
 	TSharedRef<class IHttpRequest> HttpRequest = FHttpModule::Get().CreateRequest();
 
 	// Download the next stream chunk
-	HttpRequest->SetURL( FString::Printf( TEXT( "%sdownload?Version=%s&Session=%s&Filename=checkpoint.%i" ), *ServerURL, *SessionVersion, *SessionName, 0 ) );
+	HttpRequest->SetURL( FString::Printf( TEXT( "%sdownload?Version=%s&Session=%s&Filename=checkpoint.%i" ), *ServerURL, *SessionVersion, *SessionName, TimeInMS ) );
 	HttpRequest->SetVerb( TEXT( "GET" ) );
 
 	HttpRequest->OnProcessRequestComplete().BindRaw( this, &FHttpNetworkReplayStreamer::HttpDownloadCheckpointFinished );
 
-	HttpState = EHttptate::DownloadingCheckpoint;
-
-	HttpRequest->ProcessRequest();
-
 	GotoCheckpointDelegate = Delegate;
+
+	AddRequestToQueue( EQueuedHttpRequestType::DownloadingCheckpoint, HttpRequest );
 }
 
 void FHttpNetworkReplayStreamer::FlushCheckpointInternal( uint32 TimeInMS )
 {
-	if ( SessionName.IsEmpty() || !CheckpointArchive.ArIsSaving )
+	if ( SessionName.IsEmpty() || StreamerState != EStreamerState::StreamingUp || CheckpointArchive.Buffer.Num() == 0 )
 	{
-		// IF there is no active session, or we are not recording, we don't need to flush
-		return;
-	}
-
-	if ( IsHttpBusy() )
-	{
-		// If we're currently busy we can't flush right now
+		// If there is no active session, or we are not recording, we don't need to flush
+		CheckpointArchive.Buffer.Empty();
+		CheckpointArchive.Pos = 0;
 		return;
 	}
 
@@ -357,12 +367,10 @@ void FHttpNetworkReplayStreamer::FlushCheckpointInternal( uint32 TimeInMS )
 
 	HttpRequest->OnProcessRequestComplete().BindRaw( this, &FHttpNetworkReplayStreamer::HttpUploadFinished );
 
-	HttpState = EHttptate::UploadingStream;
-
 	// Save the chunk to the checkpoint
 	CheckpointArchive << StreamFileCount;
 
-	HttpRequest->SetURL( FString::Printf( TEXT( "%supload?Version=%s&Session=%s&NumChunks=%i&Time=%i&Filename=checkpoint.%i" ), *ServerURL, *SessionVersion, *SessionName, StreamFileCount, TotalDemoTimeInMS, 0 ) );
+	HttpRequest->SetURL( FString::Printf( TEXT( "%supload?Version=%s&Session=%s&NumChunks=%i&Time=%i&Filename=checkpoint.%i" ), *ServerURL, *SessionVersion, *SessionName, StreamFileCount, TotalDemoTimeInMS, TimeInMS ) );
 	HttpRequest->SetVerb( TEXT( "POST" ) );
 	HttpRequest->SetHeader( TEXT( "Content-Type" ), TEXT( "application/octet-stream" ) );
 	HttpRequest->SetContent( CheckpointArchive.Buffer );
@@ -370,7 +378,7 @@ void FHttpNetworkReplayStreamer::FlushCheckpointInternal( uint32 TimeInMS )
 	CheckpointArchive.Buffer.Empty();
 	CheckpointArchive.Pos = 0;
 
-	HttpRequest->ProcessRequest();
+	AddRequestToQueue( EQueuedHttpRequestType::UploadingCheckpoint, HttpRequest );
 }
 
 void FHttpNetworkReplayStreamer::DownloadHeader()
@@ -383,9 +391,7 @@ void FHttpNetworkReplayStreamer::DownloadHeader()
 
 	HttpRequest->OnProcessRequestComplete().BindRaw( this, &FHttpNetworkReplayStreamer::HttpDownloadHeaderFinished );
 
-	HttpState = EHttptate::DownloadingHeader;
-
-	HttpRequest->ProcessRequest();
+	AddRequestToQueue( EQueuedHttpRequestType::DownloadingHeader, HttpRequest );
 }
 
 void FHttpNetworkReplayStreamer::DownloadNextChunk()
@@ -395,13 +401,17 @@ void FHttpNetworkReplayStreamer::DownloadNextChunk()
 		return;
 	}
 
+	if ( GotoCheckpointDelegate.IsBound() )
+	{
+		// Don't download stream while we're waiting for a checkpoint to download
+		return;
+	}
+
 	const bool bMoreChunksDefinitelyAvilable = StreamFileCount < NumDownloadChunks;
 	const bool bMoreChunksMayBeAvilable = bStreamIsLive;
 
 	if ( !bMoreChunksDefinitelyAvilable && !bMoreChunksMayBeAvilable )
 	{
-		//StreamerState = EStreamerState::Idle;
-		//UE_LOG( LogHttpReplay, Log, TEXT( "FHttpNetworkReplayStreamer::DownloadNextChunk. Done streaming." ) );
 		return;
 	}
 
@@ -425,36 +435,19 @@ void FHttpNetworkReplayStreamer::DownloadNextChunk()
 
 	HttpRequest->OnProcessRequestComplete().BindRaw( this, &FHttpNetworkReplayStreamer::HttpDownloadFinished );
 
-	HttpState = EHttptate::DownloadingStream;
-
-	HttpRequest->ProcessRequest();
+	AddRequestToQueue( EQueuedHttpRequestType::DownloadingStream, HttpRequest );
 
 	LastChunkTime = FPlatformTime::Seconds();
 }
 
-void FHttpNetworkReplayStreamer::RefreshViewer()
+void FHttpNetworkReplayStreamer::RefreshViewer( const bool bFinal )
 {
-	if ( IsHttpBusy() )
-	{
-		return;
-	}
-
-	// Determine if it's time to download the next chunk
-	const double REFRESH_VIEWER_IN_SECONDS = 10;
-
-	if ( StreamerState != EStreamerState::StreamingDownFinal && FPlatformTime::Seconds() - LastRefreshViewerTime < REFRESH_VIEWER_IN_SECONDS )
-	{
-		// Not time yet
-		return;
-	}
-
 	TSharedRef<class IHttpRequest> HttpRequest = FHttpModule::Get().CreateRequest();
 
 	// Download the next stream chunk
-	if ( StreamerState == EStreamerState::StreamingDownFinal )
+	if ( bFinal )
 	{
 		HttpRequest->SetURL( FString::Printf( TEXT( "%srefreshviewer?Version=%s&Session=%s&Viewer=%s&Final=true" ), *ServerURL, *SessionVersion, *SessionName, *ViewerName ) );
-		StreamerState = EStreamerState::Idle;
 	}
 	else
 	{
@@ -465,15 +458,24 @@ void FHttpNetworkReplayStreamer::RefreshViewer()
 
 	HttpRequest->OnProcessRequestComplete().BindRaw( this, &FHttpNetworkReplayStreamer::HttpRefreshViewerFinished );
 
-	HttpState = EHttptate::RefreshingViewer;
-
-	HttpRequest->ProcessRequest();
+	AddRequestToQueue( EQueuedHttpRequestType::RefreshingViewer, HttpRequest );
 
 	LastRefreshViewerTime = FPlatformTime::Seconds();
 }
 
 void FHttpNetworkReplayStreamer::SetLastError( const ENetworkReplayError::Type InLastError )
 {
+	// Cancel any in flight request
+	if ( InFlightHttpRequest.IsValid() )
+	{
+		InFlightHttpRequest->Request->CancelRequest();
+		InFlightHttpRequest = NULL;
+	}
+
+	// Empty the request queue
+	TSharedPtr< FQueuedHttpRequest > QueuedRequest;
+	while ( QueuedHttpRequests.Dequeue( QueuedRequest ) ) { }
+
 	StreamerState		= EStreamerState::Idle;
 	StreamerLastError	= InLastError;
 }
@@ -568,17 +570,18 @@ void FHttpNetworkReplayStreamer::EnumerateStreams( const FString& VersionString,
 	HttpRequest->OnProcessRequestComplete().BindRaw( this, &FHttpNetworkReplayStreamer::HttpEnumerateSessionsFinished );
 
 	EnumerateStreamsDelegate = Delegate;
-	HttpState = EHttptate::EnumeratingSessions;
 
-	HttpRequest->ProcessRequest();
+	AddRequestToQueue( EQueuedHttpRequestType::EnumeratingSessions, HttpRequest );
 }
 
 void FHttpNetworkReplayStreamer::HttpStartUploadingFinished( FHttpRequestPtr HttpRequest, FHttpResponsePtr HttpResponse, bool bSucceeded )
 {
-	check( HttpState == EHttptate::StartUploading );
-	check( StreamerState == EStreamerState::NeedToUploadHeader );
+	check( InFlightHttpRequest.IsValid() );
+	check( InFlightHttpRequest->Request == HttpRequest );
+	check( InFlightHttpRequest->Type == EQueuedHttpRequestType::StartUploading );
+	check( StreamerState == EStreamerState::StreamingUp );
 
-	HttpState = EHttptate::Idle;
+	InFlightHttpRequest = NULL;
 
 	if ( bSucceeded && HttpResponse->GetResponseCode() == EHttpResponseCodes::Ok )
 	{
@@ -595,11 +598,12 @@ void FHttpNetworkReplayStreamer::HttpStartUploadingFinished( FHttpRequestPtr Htt
 
 void FHttpNetworkReplayStreamer::HttpStopUploadingFinished( FHttpRequestPtr HttpRequest, FHttpResponsePtr HttpResponse, bool bSucceeded )
 {
-	check( HttpState == EHttptate::StopUploading );
-	check( StreamerState == EStreamerState::StreamingUpFinal );
+	check( InFlightHttpRequest.IsValid() );
+	check( InFlightHttpRequest->Request == HttpRequest );
+	check( InFlightHttpRequest->Type == EQueuedHttpRequestType::StopUploading );
+	check( StreamerState == EStreamerState::StreamingUp );
 
-	HttpState = EHttptate::Idle;
-	StreamerState = EStreamerState::Idle;
+	InFlightHttpRequest = NULL;
 
 	if ( bSucceeded && HttpResponse->GetResponseCode() == EHttpResponseCodes::Ok )
 	{
@@ -616,22 +620,21 @@ void FHttpNetworkReplayStreamer::HttpStopUploadingFinished( FHttpRequestPtr Http
 	StreamArchive.Buffer.Empty();
 	StreamArchive.Pos = 0;
 	StreamFileCount = 0;
-	bStopStreamingCalled = false;
 	SessionName.Empty();
 }
 
 void FHttpNetworkReplayStreamer::HttpHeaderUploadFinished( FHttpRequestPtr HttpRequest, FHttpResponsePtr HttpResponse, bool bSucceeded )
 {
-	check( HttpState == EHttptate::UploadingHeader );
-	check( StreamerState == EStreamerState::NeedToUploadHeader );
+	check( InFlightHttpRequest.IsValid() );
+	check( InFlightHttpRequest->Request == HttpRequest );
+	check( InFlightHttpRequest->Type == EQueuedHttpRequestType::UploadingHeader );
+	check( StreamerState == EStreamerState::StreamingUp );
 
-	HttpState = EHttptate::Idle;
+	InFlightHttpRequest = NULL;
 
 	if ( bSucceeded && HttpResponse->GetResponseCode() == EHttpResponseCodes::Ok )
 	{
 		UE_LOG( LogHttpReplay, Log, TEXT( "FHttpNetworkReplayStreamer::HttpHeaderUploadFinished." ) );
-
-		StreamerState = EStreamerState::StreamingUp;
 	}
 	else
 	{
@@ -640,32 +643,18 @@ void FHttpNetworkReplayStreamer::HttpHeaderUploadFinished( FHttpRequestPtr HttpR
 	}
 }
 
-void FHttpNetworkReplayStreamer::HttpUploadFinished( FHttpRequestPtr , FHttpResponsePtr HttpResponse, bool bSucceeded )
+void FHttpNetworkReplayStreamer::HttpUploadFinished( FHttpRequestPtr HttpRequest, FHttpResponsePtr HttpResponse, bool bSucceeded )
 {
-	check( HttpState == EHttptate::UploadingStream );
-	check( StreamerState == EStreamerState::StreamingUp || StreamerState == EStreamerState::StreamingUpFinal );
+	check( InFlightHttpRequest.IsValid() );
+	check( InFlightHttpRequest->Request == HttpRequest );
+	check( InFlightHttpRequest->Type == EQueuedHttpRequestType::UploadingStream || InFlightHttpRequest->Type == EQueuedHttpRequestType::UploadingCheckpoint );
+	check( StreamerState == EStreamerState::StreamingUp );
 
-	HttpState = EHttptate::Idle;
+	InFlightHttpRequest = NULL;
 
 	if ( bSucceeded && HttpResponse->GetResponseCode() == EHttpResponseCodes::Ok )
 	{
 		UE_LOG( LogHttpReplay, Verbose, TEXT( "FHttpNetworkReplayStreamer::HttpUploadFinished." ) );
-
-		if ( StreamerState == EStreamerState::StreamingUpFinal )
-		{
-			// Create the Http request and add to pending request list
-			TSharedRef<class IHttpRequest> HttpRequest = FHttpModule::Get().CreateRequest();
-
-			HttpRequest->OnProcessRequestComplete().BindRaw( this, &FHttpNetworkReplayStreamer::HttpStopUploadingFinished );
-
-			HttpState = EHttptate::StopUploading;
-
-			HttpRequest->SetURL( FString::Printf( TEXT( "%sstopuploading?Version=%s&Session=%s&NumChunks=%i&Time=%i" ), *ServerURL, *SessionVersion, *SessionName, StreamFileCount, TotalDemoTimeInMS ) );
-			HttpRequest->SetVerb( TEXT( "POST" ) );
-			HttpRequest->SetHeader( TEXT( "Content-Type" ), TEXT( "application/octet-stream" ) );
-
-			HttpRequest->ProcessRequest();
-		}
 	}
 	else
 	{
@@ -676,10 +665,12 @@ void FHttpNetworkReplayStreamer::HttpUploadFinished( FHttpRequestPtr , FHttpResp
 
 void FHttpNetworkReplayStreamer::HttpStartDownloadingFinished( FHttpRequestPtr HttpRequest, FHttpResponsePtr HttpResponse, bool bSucceeded )
 {
-	check( HttpState == EHttptate::StartDownloading );
-	check( StreamerState == EStreamerState::NeedToDownloadHeader );
+	check( InFlightHttpRequest.IsValid() );
+	check( InFlightHttpRequest->Request == HttpRequest );
+	check( InFlightHttpRequest->Type == EQueuedHttpRequestType::StartDownloading );
+	check( StreamerState == EStreamerState::StreamingDown );
 
-	HttpState = EHttptate::Idle;
+	InFlightHttpRequest = NULL;
 
 	if ( bSucceeded && HttpResponse->GetResponseCode() == EHttpResponseCodes::Ok )
 	{
@@ -699,7 +690,7 @@ void FHttpNetworkReplayStreamer::HttpStartDownloadingFinished( FHttpRequestPtr H
 		// First, download the header
 		if ( NumDownloadChunks > 0 )
 		{			
-			StreamerState = EStreamerState::NeedToDownloadHeader;
+			// The next queued request should be the download header request
 		}
 		else
 		{
@@ -728,11 +719,13 @@ void FHttpNetworkReplayStreamer::HttpStartDownloadingFinished( FHttpRequestPtr H
 
 void FHttpNetworkReplayStreamer::HttpDownloadHeaderFinished( FHttpRequestPtr HttpRequest, FHttpResponsePtr HttpResponse, bool bSucceeded )
 {
-	check( HttpState == EHttptate::DownloadingHeader );
-	check( StreamerState == EStreamerState::NeedToDownloadHeader );
+	check( InFlightHttpRequest.IsValid() );
+	check( InFlightHttpRequest->Request == HttpRequest );
+	check( InFlightHttpRequest->Type == EQueuedHttpRequestType::DownloadingHeader );
+	check( StreamerState == EStreamerState::StreamingDown );
 	check( StreamArchive.IsLoading() );
 
-	HttpState = EHttptate::Idle;
+	InFlightHttpRequest = NULL;
 
 	if ( bSucceeded && HttpResponse->GetResponseCode() == EHttpResponseCodes::Ok )
 	{
@@ -741,7 +734,6 @@ void FHttpNetworkReplayStreamer::HttpDownloadHeaderFinished( FHttpRequestPtr Htt
 		UE_LOG( LogHttpReplay, Log, TEXT( "FHttpNetworkReplayStreamer::HttpDownloadHeaderFinished. Size: %i" ), HeaderArchive.Buffer.Num() );
 
 		StartStreamingDelegate.ExecuteIfBound( true, StreamArchive.IsSaving() );
-		StreamerState = EStreamerState::StreamingDown;
 	}
 	else
 	{
@@ -760,10 +752,12 @@ void FHttpNetworkReplayStreamer::HttpDownloadHeaderFinished( FHttpRequestPtr Htt
 
 void FHttpNetworkReplayStreamer::HttpDownloadFinished( FHttpRequestPtr HttpRequest, FHttpResponsePtr HttpResponse, bool bSucceeded )
 {
-	check( HttpState == EHttptate::DownloadingStream );
+	check( InFlightHttpRequest.IsValid() );
+	check( InFlightHttpRequest->Request == HttpRequest );
+	check( InFlightHttpRequest->Type == EQueuedHttpRequestType::DownloadingStream );
 	check( StreamArchive.IsLoading() );
 
-	HttpState = EHttptate::Idle;
+	InFlightHttpRequest = NULL;
 
 	if ( bSucceeded && HttpResponse->GetResponseCode() == EHttpResponseCodes::Ok )
 	{
@@ -803,8 +797,13 @@ void FHttpNetworkReplayStreamer::HttpDownloadFinished( FHttpRequestPtr HttpReque
 
 void FHttpNetworkReplayStreamer::HttpDownloadCheckpointFinished( FHttpRequestPtr HttpRequest, FHttpResponsePtr HttpResponse, bool bSucceeded )
 {
-	check( HttpState == EHttptate::DownloadingCheckpoint );
+	check( InFlightHttpRequest.IsValid() );
+	check( InFlightHttpRequest->Request == HttpRequest );
+	check( InFlightHttpRequest->Type == EQueuedHttpRequestType::DownloadingCheckpoint );
 	check( StreamArchive.IsLoading() );
+	check( GotoCheckpointDelegate.IsBound() );
+
+	InFlightHttpRequest = NULL;
 
 	if ( bSucceeded && HttpResponse->GetResponseCode() == EHttpResponseCodes::Ok )
 	{
@@ -815,9 +814,6 @@ void FHttpNetworkReplayStreamer::HttpDownloadCheckpointFinished( FHttpRequestPtr
 		CheckpointArchive.Buffer = HttpResponse->GetContent();
 		CheckpointArchive.Pos = 0;
 
-		int32 NumValues = 0;
-		CheckpointArchive << NumValues;
-
 		// Seek to the end, where we added our internal meta data
 		CheckpointArchive.Seek( CheckpointArchive.Buffer.Num() - 4 );
 		
@@ -841,17 +837,17 @@ void FHttpNetworkReplayStreamer::HttpDownloadCheckpointFinished( FHttpRequestPtr
 		GotoCheckpointDelegate.ExecuteIfBound( false );
 	}
 
-	HttpState = EHttptate::Idle;
-
 	GotoCheckpointDelegate = FOnCheckpointReadyDelegate();
 }
 
 void FHttpNetworkReplayStreamer::HttpRefreshViewerFinished( FHttpRequestPtr HttpRequest, FHttpResponsePtr HttpResponse, bool bSucceeded )
 {
-	check( HttpState == EHttptate::RefreshingViewer );
+	check( InFlightHttpRequest.IsValid() );
+	check( InFlightHttpRequest->Request == HttpRequest );
+	check( InFlightHttpRequest->Type == EQueuedHttpRequestType::RefreshingViewer );
 	check( StreamArchive.IsLoading() );
 
-	HttpState = EHttptate::Idle;
+	InFlightHttpRequest = NULL;
 
 	if ( !bSucceeded || HttpResponse->GetResponseCode() != EHttpResponseCodes::Ok )
 	{
@@ -862,9 +858,11 @@ void FHttpNetworkReplayStreamer::HttpRefreshViewerFinished( FHttpRequestPtr Http
 
 void FHttpNetworkReplayStreamer::HttpEnumerateSessionsFinished( FHttpRequestPtr HttpRequest, FHttpResponsePtr HttpResponse, bool bSucceeded )
 {
-	check( HttpState == EHttptate::EnumeratingSessions );
+	check( InFlightHttpRequest.IsValid() );
+	check( InFlightHttpRequest->Request == HttpRequest );
+	check( InFlightHttpRequest->Type == EQueuedHttpRequestType::EnumeratingSessions );
 
-	HttpState = EHttptate::Idle;
+	InFlightHttpRequest = NULL;
 
 	if ( bSucceeded && HttpResponse->GetResponseCode() == EHttpResponseCodes::Ok )
 	{		
@@ -915,45 +913,41 @@ void FHttpNetworkReplayStreamer::Tick( const float DeltaTime )
 		return;
 	}
 
-	if ( StreamFileCount >= NumDownloadChunks && !bStreamIsLive )
+	// Attempt to process the next http request
+	TSharedPtr< FQueuedHttpRequest > QueuedRequest;
+	if ( QueuedHttpRequests.Dequeue( QueuedRequest ) )
 	{
-		StreamArchive.bAtEndOfReplay = true;
-	}
+		UE_LOG( LogHttpReplay, Verbose, TEXT( "FHttpNetworkReplayStreamer::Tick. Dequeue Type: %i" ), (int)QueuedRequest->Type );
 
-	if ( SessionName.IsEmpty() )
-	{
-		check( StreamerState == EStreamerState::Idle );
+		InFlightHttpRequest = QueuedRequest;
+		InFlightHttpRequest->Request->ProcessRequest();
 		return;
 	}
 
-	if ( bStopStreamingCalled )
+	// See if we're waiting on uploading the header
+	if ( bNeedToUploadHeader )
 	{
-		// If http isn't busy and we need to flush the last chunk, then go to that state now
-		if ( StreamerState == EStreamerState::StreamingUp )
-		{
-			StreamerState = EStreamerState::StreamingUpFinal;
-			bStopStreamingCalled = false;
-		}
-		else if ( StreamerState == EStreamerState::StreamingDown )
+		check( QueuedHttpRequests.IsEmpty() );
+
+		// If we're waiting on uploading the header, don't do anything until we have the session name
+		if ( SessionName.IsEmpty() )
 		{
-			StreamerState = EStreamerState::StreamingDownFinal;
-			bStopStreamingCalled = false;
+			return;
 		}
+		UploadHeader();
+		bNeedToUploadHeader = false;
 	}
 
-	if ( StreamerState == EStreamerState::NeedToUploadHeader )
+	if ( bStopStreamingCalled )
 	{
-		// If we're waiting on the header, don't do anything until the header has data and is uploaded
-		UploadHeader();
+		// If we want to stop streaming, we can go to idle once our queue is empty
+		StreamerState = EStreamerState::Idle;
+		bStopStreamingCalled = false;
+		return;
 	}
-	else if ( StreamerState == EStreamerState::StreamingUp )
-	{
-		if ( FlushCheckpointTime > 0 )
-		{
-			FlushCheckpointInternal( FlushCheckpointTime );
-			FlushCheckpointTime = 0;
-		}
 
+	if ( StreamerState == EStreamerState::StreamingUp )
+	{
 		const double FLUSH_TIME_IN_SECONDS = 10;
 
 		if ( FPlatformTime::Seconds() - LastChunkTime > FLUSH_TIME_IN_SECONDS )
@@ -961,30 +955,27 @@ void FHttpNetworkReplayStreamer::Tick( const float DeltaTime )
 			FlushStream();
 		}
 	}
-	else if ( StreamerState == EStreamerState::StreamingUpFinal )
-	{
-		FlushStream();
-	}
-	else if ( StreamerState == EStreamerState::NeedToDownloadHeader )
-	{
-		// If we're waiting on the header to download, don't do anything until the header has been downloaded
-		DownloadHeader();
-	}
 	else if ( StreamerState == EStreamerState::StreamingDown )
 	{
-		// Give viewer refreshing priority since it's quick
-		RefreshViewer();
+		if ( StreamFileCount >= NumDownloadChunks && !bStreamIsLive )
+		{
+			StreamArchive.bAtEndOfReplay = true;
+		}
+
+		const double REFRESH_VIEWER_IN_SECONDS = 10;
+
+		if ( FPlatformTime::Seconds() - LastRefreshViewerTime > REFRESH_VIEWER_IN_SECONDS )
+		{
+			RefreshViewer( false );
+		}
+
 		DownloadNextChunk();
 	}
-	else if ( StreamerState == EStreamerState::StreamingDownFinal )
-	{
-		RefreshViewer();
-	}
 }
 
 bool FHttpNetworkReplayStreamer::IsHttpBusy() const
 {
-	return HttpState != EHttptate::Idle;
+	return InFlightHttpRequest.IsValid();
 }
 
 bool FHttpNetworkReplayStreamer::IsStreaming() const
@@ -1010,7 +1001,7 @@ void FHttpNetworkReplayStreamingFactory::Tick( float DeltaTime )
 		HttpStreamers[i]->Tick( DeltaTime );
 		
 		// We can release our hold when streaming is completely done
-		if ( HttpStreamers[i].IsUnique() && !HttpStreamers[i]->IsHttpBusy() && !HttpStreamers[i]->bStopStreamingCalled )
+		if ( HttpStreamers[i].IsUnique() && !HttpStreamers[i]->IsHttpBusy() && HttpStreamers[i]->QueuedHttpRequests.IsEmpty() )
 		{
 			if ( HttpStreamers[i]->IsStreaming() )
 			{
diff --git a/Engine/Source/Runtime/NetworkReplayStreaming/HttpNetworkReplayStreaming/Private/HttpNetworkReplayStreaming.h b/Engine/Source/Runtime/NetworkReplayStreaming/HttpNetworkReplayStreaming/Private/HttpNetworkReplayStreaming.h
--- a/Engine/Source/Runtime/NetworkReplayStreaming/HttpNetworkReplayStreaming/Private/HttpNetworkReplayStreaming.h
+++ b/Engine/Source/Runtime/NetworkReplayStreaming/HttpNetworkReplayStreaming/Private/HttpNetworkReplayStreaming.h
@@ -8,10 +8,10 @@
 /**
  * Archive used to buffer stream over http
  */
-class HttpStreamFArchive : public FArchive
+class FHttpStreamFArchive : public FArchive
 {
 public:
-	HttpStreamFArchive() : Pos( 0 ), bAtEndOfReplay( false ) {}
+	FHttpStreamFArchive() : Pos( 0 ), bAtEndOfReplay( false ) {}
 
 	virtual void	Serialize( void* V, int64 Length );
 	virtual int64	Tell();
@@ -24,6 +24,30 @@ public:
 	bool			bAtEndOfReplay;
 };
 
+enum class EQueuedHttpRequestType
+{
+	StartUploading,				// We have made a request to start uploading a replay
+	UploadingHeader,			// We are uploading the replay header
+	UploadingStream,			// We are in the process of uploading the replay stream
+	StopUploading,				// We have made the request to stop uploading a live replay stream
+	StartDownloading,			// We have made the request to start downloading a replay stream
+	DownloadingHeader,			// We are downloading the replay header
+	DownloadingStream,			// We are in the process of downloading the replay stream
+	RefreshingViewer,			// We are refreshing the server to let it know we're still viewing
+	EnumeratingSessions,		// We are in the process of downloading the available sessions
+	UploadingCheckpoint,		// We are uploading a checkpoint
+	DownloadingCheckpoint		// We are downloading a checkpoint
+};
+
+class FQueuedHttpRequest
+{
+public:
+	FQueuedHttpRequest( const EQueuedHttpRequestType InType, TSharedPtr< class IHttpRequest > InRequest ) : Type( InType ), Request( InRequest ) { }
+
+	EQueuedHttpRequestType				Type;
+	TSharedPtr< class IHttpRequest >	Request;
+};
+
 /**
  * Http network replay streaming manager
  */
@@ -54,11 +78,13 @@ public:
 	/** FHttpNetworkReplayStreamer */
 	void UploadHeader();
 	void FlushStream();
+	void StopUploading();
 	void DownloadHeader();
 	void DownloadNextChunk();
-	void RefreshViewer();
+	void RefreshViewer( const bool bFinal );
 	void SetLastError( const ENetworkReplayError::Type InLastError );
 	void FlushCheckpointInternal( uint32 TimeInMS );
+	void AddRequestToQueue( const EQueuedHttpRequestType Type, TSharedPtr< class IHttpRequest >	Request );
 
 	/** Delegates */
 	void HttpStartDownloadingFinished( FHttpRequestPtr HttpRequest, FHttpResponsePtr HttpResponse, bool bSucceeded );
@@ -76,38 +102,17 @@ public:
 	bool IsHttpBusy() const;		// True if we are waiting on http request response
 	bool IsStreaming() const;		// True if we are streaming a replay up or down
 
-	/** EHttptate - If not idle, there is a http request in flight */
-	enum class EHttptate
-	{
-		Idle,						// There is no http request in flight
-		StartUploading,				// We have made a request to start uploading a replay
-		UploadingHeader,			// We are uploading the replay header
-		UploadingStream,			// We are in the process of uploading the replay stream
-		StopUploading,				// We have made the request to stop uploading a live replay stream
-		StartDownloading,			// We have made the request to start downloading a replay stream
-		DownloadingHeader,			// We are downloading the replay header
-		DownloadingStream,			// We are in the process of downloading the replay stream
-		RefreshingViewer,			// We are refreshing the server to let it know we're still viewing
-		EnumeratingSessions,		// We are in the process of downloading the available sessions
-		UploadingCheckpoint,		// We are uploading a checkpoint
-		DownloadingCheckpoint,		// We are downloading a checkpoint
-	};
-
 	/** EStreamerState - Overall state of the streamer */
 	enum class EStreamerState
 	{
 		Idle,						// The streamer is idle. Either we haven't started streaming yet, or we are done
-		NeedToUploadHeader,			// We are waiting to upload the header
-		NeedToDownloadHeader,		// We are waiting to download the header
 		StreamingUp,				// We are in the process of streaming a replay to the http server
-		StreamingDown,				// We are in the process of streaming a replay from the http server
-		StreamingDownFinal,			// We are done streaming (will notify server to remove us as a viewer)
-		StreamingUpFinal,			// We are uploading the final stream
+		StreamingDown				// We are in the process of streaming a replay from the http server
 	};
 
-	HttpStreamFArchive		HeaderArchive;			// Archive used to buffer the header stream
-	HttpStreamFArchive		StreamArchive;			// Archive used to buffer the data stream
-	HttpStreamFArchive		CheckpointArchive;		// Archive used to buffer checkpoint data
+	FHttpStreamFArchive		HeaderArchive;			// Archive used to buffer the header stream
+	FHttpStreamFArchive		StreamArchive;			// Archive used to buffer the data stream
+	FHttpStreamFArchive		CheckpointArchive;		// Archive used to buffer checkpoint data
 	FString					SessionName;			// Name of the session on the http replay server
 	FString					SessionVersion;			// Version of the session
 	FString					ServerURL;				// The address of the server
@@ -115,12 +120,11 @@ public:
 	double					LastChunkTime;			// The last time we uploaded/downloaded a chunk
 	double					LastRefreshViewerTime;	// The last time we refreshed ourselves as an active viewer
 	EStreamerState			StreamerState;			// Overall state of the streamer
-	EHttptate				HttpState;				// If not idle, there is a http request in flight
 	bool					bStopStreamingCalled;
+	bool					bNeedToUploadHeader;	// We're waiting on session name so we can upload header
 	bool					bStreamIsLive;			// If true, we are viewing a live stream
 	int32					NumDownloadChunks;
 	uint32					TotalDemoTimeInMS;
-	uint32					FlushCheckpointTime;
 	FString					ViewerName;
 	uint32					HighPriorityEndTime;
 
@@ -129,6 +133,9 @@ public:
 	FOnStreamReadyDelegate			StartStreamingDelegate;		// Delegate passed in to StartStreaming
 	FOnEnumerateStreamsComplete		EnumerateStreamsDelegate;
 	FOnCheckpointReadyDelegate		GotoCheckpointDelegate;
+
+	TQueue< TSharedPtr< FQueuedHttpRequest > >	QueuedHttpRequests;
+	TSharedPtr< FQueuedHttpRequest >			InFlightHttpRequest;
 };
 
 class FHttpNetworkReplayStreamingFactory : public INetworkReplayStreamingFactory, public FTickableGameObject
