Http replay streamer cleanup

diff --git a/Engine/Source/Runtime/NetworkReplayStreaming/HttpNetworkReplayStreaming/Private/HttpNetworkReplayStreaming.cpp b/Engine/Source/Runtime/NetworkReplayStreaming/HttpNetworkReplayStreaming/Private/HttpNetworkReplayStreaming.cpp
--- a/Engine/Source/Runtime/NetworkReplayStreaming/HttpNetworkReplayStreaming/Private/HttpNetworkReplayStreaming.cpp
+++ b/Engine/Source/Runtime/NetworkReplayStreaming/HttpNetworkReplayStreaming/Private/HttpNetworkReplayStreaming.cpp
@@ -132,9 +132,9 @@ void FHttpNetworkReplayStreamer::StartStreaming( const FString& StreamName, bool
 		return;
 	}
 
-	if ( IsHttpBusy() )
+	if ( IsHttpRequestInFlight() )
 	{
-		UE_LOG( LogHttpReplay, Warning, TEXT( "FHttpNetworkReplayStreamer::StartStreaming. IsHttpBusy == true." ) );
+		UE_LOG( LogHttpReplay, Warning, TEXT( "FHttpNetworkReplayStreamer::StartStreaming. IsHttpRequestInFlight == true." ) );
 		return;
 	}
 
@@ -161,8 +161,11 @@ void FHttpNetworkReplayStreamer::StartStreaming( const FString& StreamName, bool
 
 	StreamFileCount = 0;
 
-	if ( StreamArchive.ArIsLoading )
+	if ( !bRecord )
 	{
+		// We are streaming down
+		StreamerState = EStreamerState::StreamingDown;
+
 		SessionName = StreamName;
 
 		// Notify the http server that we want to start downloading a replay
@@ -171,15 +174,20 @@ void FHttpNetworkReplayStreamer::StartStreaming( const FString& StreamName, bool
 
 		HttpRequest->OnProcessRequestComplete().BindRaw( this, &FHttpNetworkReplayStreamer::HttpStartDownloadingFinished );
 	
-		// We are streaming down
-		StreamerState = EStreamerState::StreamingDown;
-
 		AddRequestToQueue( EQueuedHttpRequestType::StartDownloading, HttpRequest );
 
 		DownloadHeader();
 	}
 	else
 	{
+		// We are streaming up
+		StreamerState = EStreamerState::StreamingUp;
+
+		SessionName.Empty();
+
+		// We can't upload the header until we have the session name, which depends on the StartUploading task above to finish
+		bNeedToUploadHeader = true;
+
 		// Notify the http server that we want to start uploading a replay
 		HttpRequest->SetURL( FString::Printf( TEXT( "%sstartuploading?Version=%s&Friendly=%s" ), *ServerURL, *SessionVersion, *StreamName ) );
 
@@ -187,15 +195,7 @@ void FHttpNetworkReplayStreamer::StartStreaming( const FString& StreamName, bool
 
 		HttpRequest->OnProcessRequestComplete().BindRaw( this, &FHttpNetworkReplayStreamer::HttpStartUploadingFinished );
 
-		SessionName.Empty();
-
-		// We are streaming up
-		StreamerState = EStreamerState::StreamingUp;
-
 		AddRequestToQueue( EQueuedHttpRequestType::StartUploading, HttpRequest );
-
-		// We can't upload the header until we have the session name, which depends on the StartUploading task above to finish
-		bNeedToUploadHeader = true;
 	}
 }
 
@@ -214,6 +214,12 @@ void FHttpNetworkReplayStreamer::StopStreaming()
 		return;
 	}
 
+	if ( bStopStreamingCalled )
+	{
+		UE_LOG( LogHttpReplay, Warning, TEXT( "FHttpNetworkReplayStreamer::StopStreaming. Already called" ) );
+		return;
+	}
+
 	bStopStreamingCalled = true;
 
 	if ( StreamerState == EStreamerState::StreamingDown )
@@ -232,15 +238,19 @@ void FHttpNetworkReplayStreamer::StopStreaming()
 
 void FHttpNetworkReplayStreamer::UploadHeader()
 {
-	if ( SessionName.IsEmpty() || !HeaderArchive.ArIsSaving )
+	check( StreamArchive.ArIsSaving );
+
+	if ( SessionName.IsEmpty() )
 	{
 		// IF there is no active session, or we are not recording, we don't need to flush
+		UE_LOG( LogHttpReplay, Warning, TEXT( "FHttpNetworkReplayStreamer::UploadHeader. No session name!" ) );
 		return;
 	}
 
 	if ( HeaderArchive.Buffer.Num() == 0 )
 	{
 		// Header wasn't serialized
+		UE_LOG( LogHttpReplay, Warning, TEXT( "FHttpNetworkReplayStreamer::UploadHeader. No header to upload" ) );
 		return;
 	}
 
@@ -270,9 +280,12 @@ void FHttpNetworkReplayStreamer::UploadHeader()
 
 void FHttpNetworkReplayStreamer::FlushStream()
 {
-	if ( bNeedToUploadHeader || !StreamArchive.ArIsSaving )
+	check( StreamArchive.ArIsSaving );
+
+	if ( bNeedToUploadHeader )
 	{
 		// If we haven't uploaded the header, or we are not recording, we don't need to flush
+		UE_LOG( LogHttpReplay, Warning, TEXT( "FHttpNetworkReplayStreamer::FlushStream. Waiting on header upload." ) );
 		return;
 	}
 
@@ -288,7 +301,7 @@ void FHttpNetworkReplayStreamer::FlushStream()
 	// Create the Http request and add to pending request list
 	TSharedRef<class IHttpRequest> HttpRequest = FHttpModule::Get().CreateRequest();
 
-	HttpRequest->OnProcessRequestComplete().BindRaw( this, &FHttpNetworkReplayStreamer::HttpUploadFinished );
+	HttpRequest->OnProcessRequestComplete().BindRaw( this, &FHttpNetworkReplayStreamer::HttpUploadStreamFinished );
 
 	HttpRequest->SetURL( FString::Printf( TEXT( "%supload?Version=%s&Session=%s&NumChunks=%i&Time=%i&Filename=stream.%i" ), *ServerURL, *SessionVersion, *SessionName, StreamFileCount + 1, TotalDemoTimeInMS, StreamFileCount ) );
 	HttpRequest->SetVerb( TEXT( "POST" ) );
@@ -305,10 +318,20 @@ void FHttpNetworkReplayStreamer::FlushStream()
 	LastChunkTime = FPlatformTime::Seconds();
 }
 
+void FHttpNetworkReplayStreamer::ConditionallyFlushStream()
+{
+	const double FLUSH_TIME_IN_SECONDS = 10;
+
+	if ( FPlatformTime::Seconds() - LastChunkTime > FLUSH_TIME_IN_SECONDS )
+	{
+		FlushStream();
+	}
+};
+
 void FHttpNetworkReplayStreamer::StopUploading()
 {
 	// Create the Http request and add to pending request list
-	TSharedRef<class IHttpRequest> HttpRequest = FHttpModule::Get().CreateRequest();
+	TSharedRef< IHttpRequest > HttpRequest = FHttpModule::Get().CreateRequest();
 
 	HttpRequest->OnProcessRequestComplete().BindRaw( this, &FHttpNetworkReplayStreamer::HttpStopUploadingFinished );
 
@@ -321,7 +344,11 @@ void FHttpNetworkReplayStreamer::StopUploading()
 
 void FHttpNetworkReplayStreamer::FlushCheckpoint( const uint32 TimeInMS )
 {
-	check( CheckpointArchive.Buffer.Num() > 0 );
+	if ( CheckpointArchive.Buffer.Num() == 0 )
+	{
+		UE_LOG( LogHttpReplay, Warning, TEXT( "FHttpNetworkReplayStreamer::FlushCheckpoint. Checkpoint is empty." ) );
+		return;
+	}
 	
 	// Flush any existing stream, we need checkpoints to line up with the next chunk
 	FlushStream();
@@ -333,6 +360,8 @@ void FHttpNetworkReplayStreamer::GotoCheckpoint( const uint32 TimeInMS, const FO
 {
 	if ( GotoCheckpointDelegate.IsBound() )
 	{
+		// If we're currently going to a checkpoint now, ignore this request
+		UE_LOG( LogHttpReplay, Warning, TEXT( "FHttpNetworkReplayStreamer::GotoCheckpoint. Busy processing another checkpoint." ) );
 		return;
 	}
 
@@ -360,12 +389,12 @@ void FHttpNetworkReplayStreamer::FlushCheckpointInternal( uint32 TimeInMS )
 	}
 
 	// Upload any new streamed data to the http server
-	UE_LOG( LogHttpReplay, Log, TEXT( "FHttpNetworkReplayStreamer::FlushCheckpoint. Size: %i, StreamFileCount: %i" ), CheckpointArchive.Buffer.Num(), StreamFileCount );
+	UE_LOG( LogHttpReplay, Log, TEXT( "FHttpNetworkReplayStreamer::FlushCheckpointInternal. Size: %i, StreamFileCount: %i" ), CheckpointArchive.Buffer.Num(), StreamFileCount );
 
 	// Create the Http request and add to pending request list
 	TSharedRef<class IHttpRequest> HttpRequest = FHttpModule::Get().CreateRequest();
 
-	HttpRequest->OnProcessRequestComplete().BindRaw( this, &FHttpNetworkReplayStreamer::HttpUploadFinished );
+	HttpRequest->OnProcessRequestComplete().BindRaw( this, &FHttpNetworkReplayStreamer::HttpUploadCheckpointFinished );
 
 	// Save the chunk to the checkpoint
 	CheckpointArchive << StreamFileCount;
@@ -394,10 +423,11 @@ void FHttpNetworkReplayStreamer::DownloadHeader()
 	AddRequestToQueue( EQueuedHttpRequestType::DownloadingHeader, HttpRequest );
 }
 
-void FHttpNetworkReplayStreamer::DownloadNextChunk()
+void FHttpNetworkReplayStreamer::ConditionallyDownloadNextChunk()
 {
-	if ( IsHttpBusy() )
+	if ( IsHttpRequestInFlight() )
 	{
+		// We never download next chunk if there is an http request in flight
 		return;
 	}
 
@@ -463,6 +493,16 @@ void FHttpNetworkReplayStreamer::RefreshViewer( const bool bFinal )
 	LastRefreshViewerTime = FPlatformTime::Seconds();
 }
 
+void FHttpNetworkReplayStreamer::ConditionallyRefreshViewer()
+{
+	const double REFRESH_VIEWER_IN_SECONDS = 10;
+
+	if ( FPlatformTime::Seconds() - LastRefreshViewerTime > REFRESH_VIEWER_IN_SECONDS )
+	{
+		RefreshViewer( false );
+	}
+};
+
 void FHttpNetworkReplayStreamer::SetLastError( const ENetworkReplayError::Type InLastError )
 {
 	// Cancel any in flight request
@@ -487,7 +527,7 @@ ENetworkReplayError::Type FHttpNetworkReplayStreamer::GetLastError() const
 
 FArchive* FHttpNetworkReplayStreamer::GetHeaderArchive()
 {
-	return ( HeaderArchive.IsSaving() || HeaderArchive.Pos < HeaderArchive.Buffer.Num() ) ? &HeaderArchive : NULL;
+	return &HeaderArchive;
 }
 
 FArchive* FHttpNetworkReplayStreamer::GetStreamingArchive()
@@ -500,6 +540,7 @@ FArchive* FHttpNetworkReplayStreamer::GetCheckpointArchive()
 	if ( bNeedToUploadHeader )
 	{
 		// If we need to upload the header, we're not ready to save checkpoints
+		// NOTE - The code needs to be resilient to this, and keep trying!!!!
 		return NULL;
 	}
 
@@ -561,7 +602,7 @@ void FHttpNetworkReplayStreamer::DeleteFinishedStream( const FString& StreamName
 
 void FHttpNetworkReplayStreamer::EnumerateStreams( const FString& VersionString, const FOnEnumerateStreamsComplete& Delegate )
 {
-	if ( IsHttpBusy() )
+	if ( IsHttpRequestInFlight() )
 	{
 		Delegate.ExecuteIfBound( TArray<FNetworkReplayStreamInfo>() );
 		return;
@@ -580,14 +621,19 @@ void FHttpNetworkReplayStreamer::EnumerateStreams( const FString& VersionString,
 	AddRequestToQueue( EQueuedHttpRequestType::EnumeratingSessions, HttpRequest );
 }
 
-void FHttpNetworkReplayStreamer::HttpStartUploadingFinished( FHttpRequestPtr HttpRequest, FHttpResponsePtr HttpResponse, bool bSucceeded )
+void FHttpNetworkReplayStreamer::RequestFinished( EStreamerState ExpectedStreamerState, EQueuedHttpRequestType ExpectedType, FHttpRequestPtr HttpRequest )
 {
+	check( StreamerState == ExpectedStreamerState );
 	check( InFlightHttpRequest.IsValid() );
 	check( InFlightHttpRequest->Request == HttpRequest );
-	check( InFlightHttpRequest->Type == EQueuedHttpRequestType::StartUploading );
-	check( StreamerState == EStreamerState::StreamingUp );
+	check( InFlightHttpRequest->Type == ExpectedType );
 
 	InFlightHttpRequest = NULL;
+};
+
+void FHttpNetworkReplayStreamer::HttpStartUploadingFinished( FHttpRequestPtr HttpRequest, FHttpResponsePtr HttpResponse, bool bSucceeded )
+{
+	RequestFinished( EStreamerState::StreamingUp, EQueuedHttpRequestType::StartUploading, HttpRequest );
 
 	if ( bSucceeded && HttpResponse->GetResponseCode() == EHttpResponseCodes::Ok )
 	{
@@ -604,12 +650,7 @@ void FHttpNetworkReplayStreamer::HttpStartUploadingFinished( FHttpRequestPtr Htt
 
 void FHttpNetworkReplayStreamer::HttpStopUploadingFinished( FHttpRequestPtr HttpRequest, FHttpResponsePtr HttpResponse, bool bSucceeded )
 {
-	check( InFlightHttpRequest.IsValid() );
-	check( InFlightHttpRequest->Request == HttpRequest );
-	check( InFlightHttpRequest->Type == EQueuedHttpRequestType::StopUploading );
-	check( StreamerState == EStreamerState::StreamingUp );
-
-	InFlightHttpRequest = NULL;
+	RequestFinished( EStreamerState::StreamingUp, EQueuedHttpRequestType::StopUploading, HttpRequest );
 
 	if ( bSucceeded && HttpResponse->GetResponseCode() == EHttpResponseCodes::Ok )
 	{
@@ -631,13 +672,9 @@ void FHttpNetworkReplayStreamer::HttpStopUploadingFinished( FHttpRequestPtr Http
 
 void FHttpNetworkReplayStreamer::HttpHeaderUploadFinished( FHttpRequestPtr HttpRequest, FHttpResponsePtr HttpResponse, bool bSucceeded )
 {
-	check( InFlightHttpRequest.IsValid() );
-	check( InFlightHttpRequest->Request == HttpRequest );
-	check( InFlightHttpRequest->Type == EQueuedHttpRequestType::UploadingHeader );
-	check( StreamerState == EStreamerState::StreamingUp );
-	check( StartStreamingDelegate.IsBound() );
+	RequestFinished( EStreamerState::StreamingUp, EQueuedHttpRequestType::UploadingHeader, HttpRequest );
 
-	InFlightHttpRequest = NULL;
+	check( StartStreamingDelegate.IsBound() );
 
 	if ( bSucceeded && HttpResponse->GetResponseCode() == EHttpResponseCodes::Ok )
 	{
@@ -656,34 +693,39 @@ void FHttpNetworkReplayStreamer::HttpHeaderUploadFinished( FHttpRequestPtr HttpR
 	StartStreamingDelegate = FOnStreamReadyDelegate();
 }
 
-void FHttpNetworkReplayStreamer::HttpUploadFinished( FHttpRequestPtr HttpRequest, FHttpResponsePtr HttpResponse, bool bSucceeded )
+void FHttpNetworkReplayStreamer::HttpUploadStreamFinished( FHttpRequestPtr HttpRequest, FHttpResponsePtr HttpResponse, bool bSucceeded )
 {
-	check( InFlightHttpRequest.IsValid() );
-	check( InFlightHttpRequest->Request == HttpRequest );
-	check( InFlightHttpRequest->Type == EQueuedHttpRequestType::UploadingStream || InFlightHttpRequest->Type == EQueuedHttpRequestType::UploadingCheckpoint );
-	check( StreamerState == EStreamerState::StreamingUp );
+	RequestFinished( EStreamerState::StreamingUp, EQueuedHttpRequestType::UploadingStream, HttpRequest );
 
-	InFlightHttpRequest = NULL;
+	if ( bSucceeded && HttpResponse->GetResponseCode() == EHttpResponseCodes::Ok )
+	{
+		UE_LOG( LogHttpReplay, Verbose, TEXT( "FHttpNetworkReplayStreamer::HttpUploadStreamFinished." ) );
+	}
+	else
+	{
+		UE_LOG( LogHttpReplay, Error, TEXT( "FHttpNetworkReplayStreamer::HttpUploadStreamFinished. FAILED" ) );
+		SetLastError( ENetworkReplayError::ServiceUnavailable );
+	}
+}
+
+void FHttpNetworkReplayStreamer::HttpUploadCheckpointFinished( FHttpRequestPtr HttpRequest, FHttpResponsePtr HttpResponse, bool bSucceeded )
+{
+	RequestFinished( EStreamerState::StreamingUp, EQueuedHttpRequestType::UploadingCheckpoint, HttpRequest );
 
 	if ( bSucceeded && HttpResponse->GetResponseCode() == EHttpResponseCodes::Ok )
 	{
-		UE_LOG( LogHttpReplay, Verbose, TEXT( "FHttpNetworkReplayStreamer::HttpUploadFinished." ) );
+		UE_LOG( LogHttpReplay, Verbose, TEXT( "FHttpNetworkReplayStreamer::HttpUploadCheckpointFinished." ) );
 	}
 	else
 	{
-		UE_LOG( LogHttpReplay, Error, TEXT( "FHttpNetworkReplayStreamer::HttpUploadFinished. FAILED" ) );
+		UE_LOG( LogHttpReplay, Error, TEXT( "FHttpNetworkReplayStreamer::HttpUploadCheckpointFinished. FAILED" ) );
 		SetLastError( ENetworkReplayError::ServiceUnavailable );
 	}
 }
 
 void FHttpNetworkReplayStreamer::HttpStartDownloadingFinished( FHttpRequestPtr HttpRequest, FHttpResponsePtr HttpResponse, bool bSucceeded )
 {
-	check( InFlightHttpRequest.IsValid() );
-	check( InFlightHttpRequest->Request == HttpRequest );
-	check( InFlightHttpRequest->Type == EQueuedHttpRequestType::StartDownloading );
-	check( StreamerState == EStreamerState::StreamingDown );
-
-	InFlightHttpRequest = NULL;
+	RequestFinished( EStreamerState::StreamingDown, EQueuedHttpRequestType::StartDownloading, HttpRequest );
 
 	if ( bSucceeded && HttpResponse->GetResponseCode() == EHttpResponseCodes::Ok )
 	{
@@ -732,15 +774,11 @@ void FHttpNetworkReplayStreamer::HttpStartDownloadingFinished( FHttpRequestPtr H
 
 void FHttpNetworkReplayStreamer::HttpDownloadHeaderFinished( FHttpRequestPtr HttpRequest, FHttpResponsePtr HttpResponse, bool bSucceeded )
 {
-	check( InFlightHttpRequest.IsValid() );
-	check( InFlightHttpRequest->Request == HttpRequest );
-	check( InFlightHttpRequest->Type == EQueuedHttpRequestType::DownloadingHeader );
-	check( StreamerState == EStreamerState::StreamingDown );
+	RequestFinished( EStreamerState::StreamingDown, EQueuedHttpRequestType::DownloadingHeader, HttpRequest );
+
 	check( StreamArchive.IsLoading() );
 	check( StartStreamingDelegate.IsBound() );
 
-	InFlightHttpRequest = NULL;
-
 	if ( bSucceeded && HttpResponse->GetResponseCode() == EHttpResponseCodes::Ok )
 	{
 		HeaderArchive.Buffer.Append( HttpResponse->GetContent() );
@@ -766,12 +804,9 @@ void FHttpNetworkReplayStreamer::HttpDownloadHeaderFinished( FHttpRequestPtr Htt
 
 void FHttpNetworkReplayStreamer::HttpDownloadFinished( FHttpRequestPtr HttpRequest, FHttpResponsePtr HttpResponse, bool bSucceeded )
 {
-	check( InFlightHttpRequest.IsValid() );
-	check( InFlightHttpRequest->Request == HttpRequest );
-	check( InFlightHttpRequest->Type == EQueuedHttpRequestType::DownloadingStream );
-	check( StreamArchive.IsLoading() );
+	RequestFinished( EStreamerState::StreamingDown, EQueuedHttpRequestType::DownloadingStream, HttpRequest );
 
-	InFlightHttpRequest = NULL;
+	check( StreamArchive.IsLoading() );
 
 	if ( bSucceeded && HttpResponse->GetResponseCode() == EHttpResponseCodes::Ok )
 	{
@@ -811,14 +846,11 @@ void FHttpNetworkReplayStreamer::HttpDownloadFinished( FHttpRequestPtr HttpReque
 
 void FHttpNetworkReplayStreamer::HttpDownloadCheckpointFinished( FHttpRequestPtr HttpRequest, FHttpResponsePtr HttpResponse, bool bSucceeded )
 {
-	check( InFlightHttpRequest.IsValid() );
-	check( InFlightHttpRequest->Request == HttpRequest );
-	check( InFlightHttpRequest->Type == EQueuedHttpRequestType::DownloadingCheckpoint );
+	RequestFinished( EStreamerState::StreamingDown, EQueuedHttpRequestType::DownloadingCheckpoint, HttpRequest );
+
 	check( StreamArchive.IsLoading() );
 	check( GotoCheckpointDelegate.IsBound() );
 
-	InFlightHttpRequest = NULL;
-
 	if ( bSucceeded && HttpResponse->GetResponseCode() == EHttpResponseCodes::Ok )
 	{
 		if ( HttpResponse->GetContent().Num() == 0 )
@@ -854,7 +886,7 @@ void FHttpNetworkReplayStreamer::HttpDownloadCheckpointFinished( FHttpRequestPtr
 	else
 	{
 		UE_LOG( LogHttpReplay, Error, TEXT( "FHttpNetworkReplayStreamer::HttpDownloadCheckpointFinished. FAILED." ) );
-		SetLastError( ENetworkReplayError::ServiceUnavailable );
+		//SetLastError( ENetworkReplayError::ServiceUnavailable );
 		GotoCheckpointDelegate.ExecuteIfBound( false );
 	}
 
@@ -863,12 +895,7 @@ void FHttpNetworkReplayStreamer::HttpDownloadCheckpointFinished( FHttpRequestPtr
 
 void FHttpNetworkReplayStreamer::HttpRefreshViewerFinished( FHttpRequestPtr HttpRequest, FHttpResponsePtr HttpResponse, bool bSucceeded )
 {
-	check( InFlightHttpRequest.IsValid() );
-	check( InFlightHttpRequest->Request == HttpRequest );
-	check( InFlightHttpRequest->Type == EQueuedHttpRequestType::RefreshingViewer );
-	check( StreamArchive.IsLoading() );
-
-	InFlightHttpRequest = NULL;
+	RequestFinished( EStreamerState::StreamingDown, EQueuedHttpRequestType::RefreshingViewer, HttpRequest );
 
 	if ( !bSucceeded || HttpResponse->GetResponseCode() != EHttpResponseCodes::Ok )
 	{
@@ -927,41 +954,69 @@ void FHttpNetworkReplayStreamer::HttpEnumerateSessionsFinished( FHttpRequestPtr
 	EnumerateStreamsDelegate = FOnEnumerateStreamsComplete();
 }
 
-void FHttpNetworkReplayStreamer::Tick( const float DeltaTime )
+bool FHttpNetworkReplayStreamer::ProcessNextHttpRequest()
 {
-	if ( IsHttpBusy() )
+	if ( IsHttpRequestInFlight() )
 	{
-		return;
+		// We only process one http request at a time to keep things simple
+		return false;
 	}
 
-	// Attempt to process the next http request
 	TSharedPtr< FQueuedHttpRequest > QueuedRequest;
+
 	if ( QueuedHttpRequests.Dequeue( QueuedRequest ) )
 	{
-		UE_LOG( LogHttpReplay, Verbose, TEXT( "FHttpNetworkReplayStreamer::Tick. Dequeue Type: %i" ), (int)QueuedRequest->Type );
+		UE_LOG( LogHttpReplay, Verbose, TEXT( "FHttpNetworkReplayStreamer::ProcessNextHttpRequest. Dequeue Type: %i" ), (int)QueuedRequest->Type );
+
+		check( !InFlightHttpRequest.IsValid() );
 
 		InFlightHttpRequest = QueuedRequest;
 		InFlightHttpRequest->Request->ProcessRequest();
+
+		// We can return now since we know a http request is now in flight
+		return true;
+	}
+
+	return false;
+}
+
+void FHttpNetworkReplayStreamer::Tick( const float DeltaTime )
+{
+	if ( IsHttpRequestInFlight() )
+	{
+		// We only process one http request at a time to keep things simple
 		return;
 	}
 
+	// Attempt to process the next http request
+	if ( ProcessNextHttpRequest() )
+	{
+		// We can now return since we know there is a request in flight
+		check( IsHttpRequestInFlight() );
+		return;
+	}
+
+	// We should have no pending or requests in flight at this point
+	check( !HasPendingHttpRequests() );
+
 	// See if we're waiting on uploading the header
+	// We have to do it this way, because the header depends on the session name, and we don't have that until
 	if ( bNeedToUploadHeader )
 	{
-		check( QueuedHttpRequests.IsEmpty() );
-
-		// If we're waiting on uploading the header, don't do anything until we have the session name
-		if ( SessionName.IsEmpty() )
+		// We can't upload the header until we have the session name, so keep waiting on it
+		if ( !SessionName.IsEmpty() )
 		{
-			return;
+			UploadHeader();
+			bNeedToUploadHeader = false;
 		}
-		UploadHeader();
-		bNeedToUploadHeader = false;
+
+		return;
 	}
 
 	if ( bStopStreamingCalled )
 	{
-		// If we want to stop streaming, we can go to idle once our queue is empty
+		// If we have a pending stop streaming request, we can fully stop once all http requests are processed
+		check( IsStreaming() );
 		StreamerState = EStreamerState::Idle;
 		bStopStreamingCalled = false;
 		return;
@@ -969,12 +1024,7 @@ void FHttpNetworkReplayStreamer::Tick( const float DeltaTime )
 
 	if ( StreamerState == EStreamerState::StreamingUp )
 	{
-		const double FLUSH_TIME_IN_SECONDS = 10;
-
-		if ( FPlatformTime::Seconds() - LastChunkTime > FLUSH_TIME_IN_SECONDS )
-		{
-			FlushStream();
-		}
+		ConditionallyFlushStream();
 	}
 	else if ( StreamerState == EStreamerState::StreamingDown )
 	{
@@ -983,22 +1033,22 @@ void FHttpNetworkReplayStreamer::Tick( const float DeltaTime )
 			StreamArchive.bAtEndOfReplay = true;
 		}
 
-		const double REFRESH_VIEWER_IN_SECONDS = 10;
-
-		if ( FPlatformTime::Seconds() - LastRefreshViewerTime > REFRESH_VIEWER_IN_SECONDS )
-		{
-			RefreshViewer( false );
-		}
-
-		DownloadNextChunk();
+		ConditionallyRefreshViewer();
+		ConditionallyDownloadNextChunk();
 	}
 }
 
-bool FHttpNetworkReplayStreamer::IsHttpBusy() const
+bool FHttpNetworkReplayStreamer::IsHttpRequestInFlight() const
 {
 	return InFlightHttpRequest.IsValid();
 }
 
+bool FHttpNetworkReplayStreamer::HasPendingHttpRequests() const
+{
+	// If there is currently one in flight, or we have more to process, return true
+	return IsHttpRequestInFlight() || !QueuedHttpRequests.IsEmpty();
+}
+
 bool FHttpNetworkReplayStreamer::IsStreaming() const
 {
 	return StreamerState != EStreamerState::Idle;
@@ -1022,7 +1072,7 @@ void FHttpNetworkReplayStreamingFactory::Tick( float DeltaTime )
 		HttpStreamers[i]->Tick( DeltaTime );
 		
 		// We can release our hold when streaming is completely done
-		if ( HttpStreamers[i].IsUnique() && !HttpStreamers[i]->IsHttpBusy() && HttpStreamers[i]->QueuedHttpRequests.IsEmpty() )
+		if ( HttpStreamers[i].IsUnique() && !HttpStreamers[i]->HasPendingHttpRequests() )
 		{
 			if ( HttpStreamers[i]->IsStreaming() )
 			{
diff --git a/Engine/Source/Runtime/NetworkReplayStreaming/HttpNetworkReplayStreaming/Private/HttpNetworkReplayStreaming.h b/Engine/Source/Runtime/NetworkReplayStreaming/HttpNetworkReplayStreaming/Private/HttpNetworkReplayStreaming.h
--- a/Engine/Source/Runtime/NetworkReplayStreaming/HttpNetworkReplayStreaming/Private/HttpNetworkReplayStreaming.h
+++ b/Engine/Source/Runtime/NetworkReplayStreaming/HttpNetworkReplayStreaming/Private/HttpNetworkReplayStreaming.h
@@ -78,15 +78,27 @@ public:
 	/** FHttpNetworkReplayStreamer */
 	void UploadHeader();
 	void FlushStream();
+	void ConditionallyFlushStream();
 	void StopUploading();
 	void DownloadHeader();
-	void DownloadNextChunk();
+	void ConditionallyDownloadNextChunk();
 	void RefreshViewer( const bool bFinal );
+	void ConditionallyRefreshViewer();
 	void SetLastError( const ENetworkReplayError::Type InLastError );
 	void FlushCheckpointInternal( uint32 TimeInMS );
 	void AddRequestToQueue( const EQueuedHttpRequestType Type, TSharedPtr< class IHttpRequest >	Request );
 
+	/** EStreamerState - Overall state of the streamer */
+	enum class EStreamerState
+	{
+		Idle,						// The streamer is idle. Either we haven't started streaming yet, or we are done
+		StreamingUp,				// We are in the process of streaming a replay to the http server
+		StreamingDown				// We are in the process of streaming a replay from the http server
+	};
+
 	/** Delegates */
+	void RequestFinished( EStreamerState ExpectedStreamerState, EQueuedHttpRequestType ExpectedType, FHttpRequestPtr HttpRequest );
+
 	void HttpStartDownloadingFinished( FHttpRequestPtr HttpRequest, FHttpResponsePtr HttpResponse, bool bSucceeded );
 	void HttpDownloadHeaderFinished( FHttpRequestPtr HttpRequest, FHttpResponsePtr HttpResponse, bool bSucceeded );
 	void HttpDownloadFinished( FHttpRequestPtr HttpRequest, FHttpResponsePtr HttpResponse, bool bSucceeded );
@@ -95,20 +107,15 @@ public:
 	void HttpStartUploadingFinished( FHttpRequestPtr HttpRequest, FHttpResponsePtr HttpResponse, bool bSucceeded );
 	void HttpStopUploadingFinished( FHttpRequestPtr HttpRequest, FHttpResponsePtr HttpResponse, bool bSucceeded );
 	void HttpHeaderUploadFinished( FHttpRequestPtr HttpRequest, FHttpResponsePtr HttpResponse, bool bSucceeded );
-	void HttpUploadFinished( FHttpRequestPtr HttpRequest, FHttpResponsePtr HttpResponse, bool bSucceeded );
+	void HttpUploadStreamFinished( FHttpRequestPtr HttpRequest, FHttpResponsePtr HttpResponse, bool bSucceeded );
+	void HttpUploadCheckpointFinished( FHttpRequestPtr HttpRequest, FHttpResponsePtr HttpResponse, bool bSucceeded );
 	void HttpEnumerateSessionsFinished( FHttpRequestPtr HttpRequest, FHttpResponsePtr HttpResponse, bool bSucceeded );
 
+	bool ProcessNextHttpRequest();
 	void Tick( const float DeltaTime );
-	bool IsHttpBusy() const;		// True if we are waiting on http request response
-	bool IsStreaming() const;		// True if we are streaming a replay up or down
-
-	/** EStreamerState - Overall state of the streamer */
-	enum class EStreamerState
-	{
-		Idle,						// The streamer is idle. Either we haven't started streaming yet, or we are done
-		StreamingUp,				// We are in the process of streaming a replay to the http server
-		StreamingDown				// We are in the process of streaming a replay from the http server
-	};
+	bool IsHttpRequestInFlight() const;		// True there is an http request currently in flight
+	bool HasPendingHttpRequests() const;	// True if there is an http request in flight, or there are more to process
+	bool IsStreaming() const;				// True if we are streaming a replay up or down
 
 	FHttpStreamFArchive		HeaderArchive;			// Archive used to buffer the header stream
 	FHttpStreamFArchive		StreamArchive;			// Archive used to buffer the data stream
