Engine side http replay checkpoint work

* We now track start/end time of portion of stream that is in memory
* We now accurately download the high priority portion of the stream
* We now accurately report whether or not a specific range of the stream is downloaded
* Naming cleanup

diff --git a/Engine/Source/Runtime/NetworkReplayStreaming/HttpNetworkReplayStreaming/Private/HttpNetworkReplayStreaming.cpp b/Engine/Source/Runtime/NetworkReplayStreaming/HttpNetworkReplayStreaming/Private/HttpNetworkReplayStreaming.cpp
--- a/Engine/Source/Runtime/NetworkReplayStreaming/HttpNetworkReplayStreaming/Private/HttpNetworkReplayStreaming.cpp
+++ b/Engine/Source/Runtime/NetworkReplayStreaming/HttpNetworkReplayStreaming/Private/HttpNetworkReplayStreaming.cpp
@@ -102,7 +102,7 @@ bool FHttpStreamFArchive::AtEnd()
 }
 
 FHttpNetworkReplayStreamer::FHttpNetworkReplayStreamer() : 
-	StreamFileCount( 0 ), 
+	StreamChunkIndex( 0 ), 
 	LastChunkTime( 0 ), 
 	LastRefreshViewerTime( 0 ),
 	StreamerState( EStreamerState::Idle ), 
@@ -111,6 +111,8 @@ FHttpNetworkReplayStreamer::FHttpNetworkReplayStreamer() :
 	bStreamIsLive( false ),
 	NumDownloadChunks( 0 ),
 	TotalDemoTimeInMS( 0 ),
+	StreamTimeRangeStart( 0 ),
+	StreamTimeRangeEnd( 0 ),
 	HighPriorityEndTime( 0 ),
 	StreamerLastError( ENetworkReplayError::None )
 {
@@ -156,10 +158,15 @@ void FHttpNetworkReplayStreamer::StartStreaming( const FString& StreamName, bool
 
 	LastChunkTime = FPlatformTime::Seconds();
 
+	TotalDemoTimeInMS = 0;
+
+	StreamTimeRangeStart	= 0;
+	StreamTimeRangeEnd		= 0;
+
 	// Create the Http request and add to pending request list
 	TSharedRef<class IHttpRequest> HttpRequest = FHttpModule::Get().CreateRequest();
 
-	StreamFileCount = 0;
+	StreamChunkIndex = 0;
 
 	if ( !bRecord )
 	{
@@ -254,17 +261,17 @@ void FHttpNetworkReplayStreamer::UploadHeader()
 		return;
 	}
 
-	check( StreamFileCount == 0 );
+	check( StreamChunkIndex == 0 );
 
 	// First upload the header
-	UE_LOG( LogHttpReplay, Log, TEXT( "FHttpNetworkReplayStreamer::UploadHeader. Header. StreamFileCount: %i, Size: %i" ), StreamFileCount, StreamArchive.Buffer.Num() );
+	UE_LOG( LogHttpReplay, Log, TEXT( "FHttpNetworkReplayStreamer::UploadHeader. Header. StreamChunkIndex: %i, Size: %i" ), StreamChunkIndex, StreamArchive.Buffer.Num() );
 
 	// Create the Http request and add to pending request list
 	TSharedRef<class IHttpRequest> HttpRequest = FHttpModule::Get().CreateRequest();
 
 	HttpRequest->OnProcessRequestComplete().BindRaw( this, &FHttpNetworkReplayStreamer::HttpHeaderUploadFinished );
 
-	HttpRequest->SetURL( FString::Printf( TEXT( "%supload?Version=%s&Session=%s&NumChunks=%i&Time=%i&Filename=replay.header" ), *ServerURL, *SessionVersion, *SessionName, StreamFileCount, TotalDemoTimeInMS ) );
+	HttpRequest->SetURL( FString::Printf( TEXT( "%supload?Version=%s&Session=%s&NumChunks=%i&Time=%i&Filename=replay.header" ), *ServerURL, *SessionVersion, *SessionName, StreamChunkIndex, TotalDemoTimeInMS ) );
 	HttpRequest->SetVerb( TEXT( "POST" ) );
 	HttpRequest->SetHeader( TEXT( "Content-Type" ), TEXT( "application/octet-stream" ) );
 	HttpRequest->SetContent( HeaderArchive.Buffer );
@@ -296,14 +303,14 @@ void FHttpNetworkReplayStreamer::FlushStream()
 	}
 
 	// Upload any new streamed data to the http server
-	UE_LOG( LogHttpReplay, Verbose, TEXT( "FHttpNetworkReplayStreamer::FlushStream. StreamFileCount: %i, Size: %i" ), StreamFileCount, StreamArchive.Buffer.Num() );
+	UE_LOG( LogHttpReplay, Verbose, TEXT( "FHttpNetworkReplayStreamer::FlushStream. StreamChunkIndex: %i, Size: %i" ), StreamChunkIndex, StreamArchive.Buffer.Num() );
 
 	// Create the Http request and add to pending request list
 	TSharedRef<class IHttpRequest> HttpRequest = FHttpModule::Get().CreateRequest();
 
 	HttpRequest->OnProcessRequestComplete().BindRaw( this, &FHttpNetworkReplayStreamer::HttpUploadStreamFinished );
 
-	HttpRequest->SetURL( FString::Printf( TEXT( "%supload?Version=%s&Session=%s&NumChunks=%i&Time=%i&Filename=stream.%i" ), *ServerURL, *SessionVersion, *SessionName, StreamFileCount + 1, TotalDemoTimeInMS, StreamFileCount ) );
+	HttpRequest->SetURL( FString::Printf( TEXT( "%supload?Version=%s&Session=%s&NumChunks=%i&Time=%i&MTime1=%i&MTime2=%i&Filename=stream.%i" ), *ServerURL, *SessionVersion, *SessionName, StreamChunkIndex + 1, TotalDemoTimeInMS, StreamTimeRangeStart, StreamTimeRangeEnd, StreamChunkIndex ) );
 	HttpRequest->SetVerb( TEXT( "POST" ) );
 	HttpRequest->SetHeader( TEXT( "Content-Type" ), TEXT( "application/octet-stream" ) );
 	HttpRequest->SetContent( StreamArchive.Buffer );
@@ -311,7 +318,10 @@ void FHttpNetworkReplayStreamer::FlushStream()
 	StreamArchive.Buffer.Empty();
 	StreamArchive.Pos = 0;
 
-	StreamFileCount++;
+	// Keep track of the time range we have in our buffer, so we can accurately upload that each time we submit a chunk
+	StreamTimeRangeStart = StreamTimeRangeEnd;
+
+	StreamChunkIndex++;
 
 	AddRequestToQueue( EQueuedHttpRequestType::UploadingStream, HttpRequest );
 
@@ -335,7 +345,7 @@ void FHttpNetworkReplayStreamer::StopUploading()
 
 	HttpRequest->OnProcessRequestComplete().BindRaw( this, &FHttpNetworkReplayStreamer::HttpStopUploadingFinished );
 
-	HttpRequest->SetURL( FString::Printf( TEXT( "%sstopuploading?Version=%s&Session=%s&NumChunks=%i&Time=%i" ), *ServerURL, *SessionVersion, *SessionName, StreamFileCount, TotalDemoTimeInMS ) );
+	HttpRequest->SetURL( FString::Printf( TEXT( "%sstopuploading?Version=%s&Session=%s&NumChunks=%i&Time=%i" ), *ServerURL, *SessionVersion, *SessionName, StreamChunkIndex, TotalDemoTimeInMS ) );
 	HttpRequest->SetVerb( TEXT( "POST" ) );
 	HttpRequest->SetHeader( TEXT( "Content-Type" ), TEXT( "application/octet-stream" ) );
 
@@ -389,7 +399,7 @@ void FHttpNetworkReplayStreamer::FlushCheckpointInternal( uint32 TimeInMS )
 	}
 
 	// Upload any new streamed data to the http server
-	UE_LOG( LogHttpReplay, Log, TEXT( "FHttpNetworkReplayStreamer::FlushCheckpointInternal. Size: %i, StreamFileCount: %i" ), CheckpointArchive.Buffer.Num(), StreamFileCount );
+	UE_LOG( LogHttpReplay, Log, TEXT( "FHttpNetworkReplayStreamer::FlushCheckpointInternal. Size: %i, StreamChunkIndex: %i" ), CheckpointArchive.Buffer.Num(), StreamChunkIndex );
 
 	// Create the Http request and add to pending request list
 	TSharedRef<class IHttpRequest> HttpRequest = FHttpModule::Get().CreateRequest();
@@ -397,9 +407,9 @@ void FHttpNetworkReplayStreamer::FlushCheckpointInternal( uint32 TimeInMS )
 	HttpRequest->OnProcessRequestComplete().BindRaw( this, &FHttpNetworkReplayStreamer::HttpUploadCheckpointFinished );
 
 	// Save the chunk to the checkpoint
-	CheckpointArchive << StreamFileCount;
+	CheckpointArchive << StreamChunkIndex;
 
-	HttpRequest->SetURL( FString::Printf( TEXT( "%supload?Version=%s&Session=%s&NumChunks=%i&Time=%i&Filename=checkpoint.%i" ), *ServerURL, *SessionVersion, *SessionName, StreamFileCount, TotalDemoTimeInMS, TimeInMS ) );
+	HttpRequest->SetURL( FString::Printf( TEXT( "%supload?Version=%s&Session=%s&NumChunks=%i&Time=%i&Filename=checkpoint.%i" ), *ServerURL, *SessionVersion, *SessionName, StreamChunkIndex, TotalDemoTimeInMS, TimeInMS ) );
 	HttpRequest->SetVerb( TEXT( "POST" ) );
 	HttpRequest->SetHeader( TEXT( "Content-Type" ), TEXT( "application/octet-stream" ) );
 	HttpRequest->SetContent( CheckpointArchive.Buffer );
@@ -437,8 +447,8 @@ void FHttpNetworkReplayStreamer::ConditionallyDownloadNextChunk()
 		return;
 	}
 
-	const bool bMoreChunksDefinitelyAvilable = StreamFileCount < NumDownloadChunks;
-	const bool bMoreChunksMayBeAvilable = bStreamIsLive;
+	const bool bMoreChunksDefinitelyAvilable	= StreamChunkIndex < NumDownloadChunks;		// We know for a fact there are more chunks available
+	const bool bMoreChunksMayBeAvilable			= bStreamIsLive;							// If we are live, there is a chance of more chunks even if we're at the end currently
 
 	if ( !bMoreChunksDefinitelyAvilable && !bMoreChunksMayBeAvilable )
 	{
@@ -448,19 +458,26 @@ void FHttpNetworkReplayStreamer::ConditionallyDownloadNextChunk()
 	// Determine if it's time to download the next chunk
 	const double CHECK_FOR_NEXT_CHUNK_IN_SECONDS = 5;
 
-	const bool bTimeToDownloadChunk = ( FPlatformTime::Seconds() - LastChunkTime > CHECK_FOR_NEXT_CHUNK_IN_SECONDS );
-	const bool bReallyNeedToDownloadChunk = ( StreamArchive.Pos >= StreamArchive.Buffer.Num() || HighPriorityEndTime > 0 ) && bMoreChunksDefinitelyAvilable;
+	const bool bTimeToDownloadChunk			= ( FPlatformTime::Seconds() - LastChunkTime > CHECK_FOR_NEXT_CHUNK_IN_SECONDS );	// Enough time has passed
+	const bool bApproachingEnd				= StreamArchive.Pos >= StreamArchive.Buffer.Num();									// We're getting close to the end
+	const bool bHighPriorityMode			= ( HighPriorityEndTime > 0 && StreamTimeRangeEnd < HighPriorityEndTime );			// We're within the high priority time range
+	const bool bReallyNeedToDownloadChunk	= ( bApproachingEnd || bHighPriorityMode ) && bMoreChunksDefinitelyAvilable;
 
+	// If it's either time to download the next chunk, or we need to for high priority reasons, download it now
+	// NOTE - bTimeToDownloadChunk may be true, even though bMoreChunksDefinitelyAvilable is false
+	//	This is because bMoreChunksMayBeAvilable is true, and we need to try and see if a new piece of the live stream is ready for download
 	if ( !bTimeToDownloadChunk && !bReallyNeedToDownloadChunk )
 	{
 		// Not time yet
 		return;
 	}
 
+	check( bMoreChunksDefinitelyAvilable || bMoreChunksMayBeAvilable );
+
 	TSharedRef<class IHttpRequest> HttpRequest = FHttpModule::Get().CreateRequest();
 
 	// Download the next stream chunk
-	HttpRequest->SetURL( FString::Printf( TEXT( "%sdownload?Version=%s&Session=%s&Filename=stream.%i" ), *ServerURL, *SessionVersion, *SessionName, StreamFileCount ) );
+	HttpRequest->SetURL( FString::Printf( TEXT( "%sdownload?Version=%s&Session=%s&Filename=stream.%i" ), *ServerURL, *SessionVersion, *SessionName, StreamChunkIndex ) );
 	HttpRequest->SetVerb( TEXT( "GET" ) );
 
 	HttpRequest->OnProcessRequestComplete().BindRaw( this, &FHttpNetworkReplayStreamer::HttpDownloadFinished );
@@ -555,6 +572,7 @@ FArchive* FHttpNetworkReplayStreamer::GetMetadataArchive()
 void FHttpNetworkReplayStreamer::UpdateTotalDemoTime( uint32 TimeInMS )
 {
 	TotalDemoTimeInMS = TimeInMS;
+	StreamTimeRangeEnd = TimeInMS;
 }
 
 bool FHttpNetworkReplayStreamer::IsDataAvailable() const
@@ -585,8 +603,8 @@ bool FHttpNetworkReplayStreamer::IsDataAvailableForTimeRange( const uint32 Start
 		return false;
 	}
 
-	// HACK: We only support the total time range for now
-	return StreamFileCount >= NumDownloadChunks;
+	// If the time is within the stream range we have downloaded, we will return true
+	return ( StartTimeInMS >= StreamTimeRangeStart && EndTimeInMS <= StreamTimeRangeEnd );
 }
 
 bool FHttpNetworkReplayStreamer::IsLive(const FString& StreamName) const 
@@ -666,7 +684,7 @@ void FHttpNetworkReplayStreamer::HttpStopUploadingFinished( FHttpRequestPtr Http
 	StreamArchive.ArIsSaving = false;
 	StreamArchive.Buffer.Empty();
 	StreamArchive.Pos = 0;
-	StreamFileCount = 0;
+	StreamChunkIndex = 0;
 	SessionName.Empty();
 }
 
@@ -742,13 +760,8 @@ void FHttpNetworkReplayStreamer::HttpStartDownloadingFinished( FHttpRequestPtr H
 
 		UE_LOG( LogHttpReplay, Log, TEXT( "FHttpNetworkReplayStreamer::HttpStartDownloadingFinished. Viewer: %s, State: %s, NumChunks: %i, DemoTime: %2.2f" ), *ViewerName, *State, NumDownloadChunks, (float)TotalDemoTimeInMS / 1000 );
 
-		// First, download the header
-		if ( NumDownloadChunks > 0 )
+		if ( NumDownloadChunks == 0 )
 		{			
-			// The next queued request should be the download header request
-		}
-		else
-		{
 			UE_LOG( LogHttpReplay, Warning, TEXT( "FHttpNetworkReplayStreamer::HttpStartDownloadingFinished. NO CHUNKS" ) );
 
 			StartStreamingDelegate.ExecuteIfBound( false, true );
@@ -810,24 +823,28 @@ void FHttpNetworkReplayStreamer::HttpDownloadFinished( FHttpRequestPtr HttpReque
 
 	if ( bSucceeded && HttpResponse->GetResponseCode() == EHttpResponseCodes::Ok )
 	{
-		FString NumChunksString = HttpResponse->GetHeader( TEXT( "NumChunks" ) );
-		FString DemoTimeString = HttpResponse->GetHeader( TEXT( "Time" ) );
-		FString State = HttpResponse->GetHeader( TEXT( "State" ) );
+		NumDownloadChunks = FCString::Atoi( *HttpResponse->GetHeader( TEXT( "NumChunks" ) ) );
+		TotalDemoTimeInMS = FCString::Atoi( *HttpResponse->GetHeader( TEXT( "Time" ) ) );
+		bStreamIsLive = HttpResponse->GetHeader( TEXT( "State" ) ) == TEXT( "Live" );
 
-		bStreamIsLive = State == TEXT( "Live" );
+		if ( StreamArchive.Buffer.Num() == 0 )
+		{
+			// If we haven't started streaming yet, this will be the start of this stream
+			StreamTimeRangeStart = FCString::Atoi( *HttpResponse->GetHeader( TEXT( "MTime1" ) ) );
+		}
 
-		NumDownloadChunks = FCString::Atoi( *NumChunksString );
-		TotalDemoTimeInMS = FCString::Atoi( *DemoTimeString );
+		// This is the new end of the stream
+		StreamTimeRangeEnd = FCString::Atoi( *HttpResponse->GetHeader( TEXT( "MTime2" ) ) );
 
 		if ( HttpResponse->GetContent().Num() > 0 || bStreamIsLive )
 		{
 			if ( HttpResponse->GetContent().Num() > 0 )
 			{
 				StreamArchive.Buffer.Append( HttpResponse->GetContent() );
-				StreamFileCount++;
+				StreamChunkIndex++;
 			}
 
-			UE_LOG( LogHttpReplay, Verbose, TEXT( "FHttpNetworkReplayStreamer::HttpDownloadFinished. State: %s, Progress: %i / %i, DemoTime: %2.2f" ), *State, StreamFileCount, NumDownloadChunks, (float)TotalDemoTimeInMS / 1000 );
+			UE_LOG( LogHttpReplay, Verbose, TEXT( "FHttpNetworkReplayStreamer::HttpDownloadFinished. Live: %i, Progress: %i / %i, Start: %i, End: %i, DemoTime: %2.2f" ), (int)bStreamIsLive, StreamChunkIndex, NumDownloadChunks, (int)StreamTimeRangeStart, (int)StreamTimeRangeEnd, (float)TotalDemoTimeInMS / 1000 );
 		}
 		else
 		{
@@ -867,11 +884,14 @@ void FHttpNetworkReplayStreamer::HttpDownloadCheckpointFinished( FHttpRequestPtr
 		CheckpointArchive.Buffer = HttpResponse->GetContent();
 		CheckpointArchive.Pos = 0;
 
+		StreamTimeRangeStart = 0;
+		StreamTimeRangeEnd = 0;
+
 		// Seek to the end, where we added our internal meta data
 		CheckpointArchive.Seek( CheckpointArchive.Buffer.Num() - 4 );
 		
 		// Set the next chunk to be the first chunk that represents the start of the stream after the checkpoint
-		CheckpointArchive << StreamFileCount;
+		CheckpointArchive << StreamChunkIndex;
 
 		// Make sure we read to the very end
 		check( CheckpointArchive.Pos == CheckpointArchive.Buffer.Num() );
@@ -882,11 +902,12 @@ void FHttpNetworkReplayStreamer::HttpDownloadCheckpointFinished( FHttpRequestPtr
 		CheckpointArchive.Pos = 0;
 
 		GotoCheckpointDelegate.ExecuteIfBound( true );
+
+		UE_LOG( LogHttpReplay, Verbose, TEXT( "FHttpNetworkReplayStreamer::HttpDownloadCheckpointFinished. SUCCESS. StreamChunkIndex: %i" ), StreamChunkIndex );
 	}
 	else
 	{
 		UE_LOG( LogHttpReplay, Error, TEXT( "FHttpNetworkReplayStreamer::HttpDownloadCheckpointFinished. FAILED." ) );
-		//SetLastError( ENetworkReplayError::ServiceUnavailable );
 		GotoCheckpointDelegate.ExecuteIfBound( false );
 	}
 
@@ -1028,7 +1049,7 @@ void FHttpNetworkReplayStreamer::Tick( const float DeltaTime )
 	}
 	else if ( StreamerState == EStreamerState::StreamingDown )
 	{
-		if ( StreamFileCount >= NumDownloadChunks && !bStreamIsLive )
+		if ( StreamChunkIndex >= NumDownloadChunks && !bStreamIsLive )
 		{
 			StreamArchive.bAtEndOfReplay = true;
 		}
diff --git a/Engine/Source/Runtime/NetworkReplayStreaming/HttpNetworkReplayStreaming/Private/HttpNetworkReplayStreaming.h b/Engine/Source/Runtime/NetworkReplayStreaming/HttpNetworkReplayStreaming/Private/HttpNetworkReplayStreaming.h
--- a/Engine/Source/Runtime/NetworkReplayStreaming/HttpNetworkReplayStreaming/Private/HttpNetworkReplayStreaming.h
+++ b/Engine/Source/Runtime/NetworkReplayStreaming/HttpNetworkReplayStreaming/Private/HttpNetworkReplayStreaming.h
@@ -123,7 +123,7 @@ public:
 	FString					SessionName;			// Name of the session on the http replay server
 	FString					SessionVersion;			// Version of the session
 	FString					ServerURL;				// The address of the server
-	int32					StreamFileCount;		// Used as a counter to increment the stream.x extension count
+	int32					StreamChunkIndex;		// Used as a counter to increment the stream.x extension count
 	double					LastChunkTime;			// The last time we uploaded/downloaded a chunk
 	double					LastRefreshViewerTime;	// The last time we refreshed ourselves as an active viewer
 	EStreamerState			StreamerState;			// Overall state of the streamer
@@ -132,6 +132,9 @@ public:
 	bool					bStreamIsLive;			// If true, we are viewing a live stream
 	int32					NumDownloadChunks;
 	uint32					TotalDemoTimeInMS;
+	uint32					LastTotalDemoTimeInMS;
+	uint32					StreamTimeRangeStart;
+	uint32					StreamTimeRangeEnd;
 	FString					ViewerName;
 	uint32					HighPriorityEndTime;
 
