Merge: Initiall replay fine scrubbing support.

* You can scrub to any time on the timeline
* Use demo.GotoTimeInSeconds to test
* You can now scrub to very beginning of timeline
* Ability to enumerate checkpoints as events
* Continue to enumerate checkpoints in the background for live games
* Improved comments
* Cleanup

diff --git a/Engine/Source/Runtime/Engine/Classes/Engine/DemoNetDriver.h b/Engine/Source/Runtime/Engine/Classes/Engine/DemoNetDriver.h
--- a/Engine/Source/Runtime/Engine/Classes/Engine/DemoNetDriver.h
+++ b/Engine/Source/Runtime/Engine/Classes/Engine/DemoNetDriver.h
@@ -57,16 +57,25 @@ class UDemoNetDriver
 	double		MaxRecordTime;
 	int32		RecordCountSinceFlush;
 	float		TimeToSkip;
+	float		LastGotoTimeInSeconds;
 
 	bool		bSavingCheckpoint;
 	double		LastCheckpointTime;
-	int32		CurrentCheckpointIndex;
 
 	void		SaveCheckpoint();
 
+	bool		bRestoreSpectatorPosition;
+	FVector		SpectatorLocation;
+	FRotator	SpectatorRotation;
+
+	FArchive*	GotoCheckpointArchive;
+	int64		GotoCheckpointSkipExtraTimeInMS;
+
+	void		LoadCheckpoint();
+
+
 private:
 	bool		bIsFastForwarding;
-	double		FastForwardStartSeconds;
 
 public:
 
@@ -81,7 +90,8 @@ public:
 	virtual void TickFlush( float DeltaSeconds ) override;
 	virtual void ProcessRemoteFunction( class AActor* Actor, class UFunction* Function, void* Parameters, struct FOutParmRec* OutParms, struct FFrame* Stack, class UObject* SubObject = nullptr ) override;
 	virtual bool IsAvailable() const override { return true; }
-	ENGINE_API void SkipTime(float InTimeToSkip);
+	ENGINE_API void SkipTime(const float InTimeToSkip);
+	ENGINE_API void GotoTimeInSeconds( const float TimeInSeconds );
 	bool InitConnectInternal( FString& Error );
 	virtual bool ShouldClientDestroyTearOffActors() const override;
 	virtual bool ShouldSkipRepNotifies() const override;
@@ -114,5 +124,5 @@ public:
 	void StopDemo();
 
 	void ReplayStreamingReady( bool bSuccess, bool bRecord );
-	void CheckpointReady( bool bSuccess );
+	void CheckpointReady( const bool bSuccess, const int64 SkipExtraTimeInMS );
 };
diff --git a/Engine/Source/Runtime/Engine/Private/DemoNetDriver.cpp b/Engine/Source/Runtime/Engine/Private/DemoNetDriver.cpp
--- a/Engine/Source/Runtime/Engine/Private/DemoNetDriver.cpp
+++ b/Engine/Source/Runtime/Engine/Private/DemoNetDriver.cpp
@@ -26,8 +26,9 @@ DEFINE_LOG_CATEGORY_STATIC( LogDemo, Log, All );
 static TAutoConsoleVariable<float> CVarDemoRecordHz( TEXT( "demo.RecordHz" ), 10, TEXT( "Number of demo frames recorded per second" ) );
 static TAutoConsoleVariable<float> CVarDemoTimeDilation( TEXT( "demo.TimeDilation" ), -1.0f, TEXT( "Override time dilation during demo playback (-1 = don't override)" ) );
 static TAutoConsoleVariable<float> CVarDemoSkipTime( TEXT( "demo.SkipTime" ), 0, TEXT( "Skip fixed amount of network replay time (in seconds)" ) );
-static TAutoConsoleVariable<int32> CVarEnableCheckpoints( TEXT( "demo.EnableCheckpoints" ), 0, TEXT( "Whether or not checkpoints save on the server" ) );
-static TAutoConsoleVariable<int32> CVarTestCheckpoint( TEXT( "demo.TestCheckpoint" ), -1, TEXT( "For testing only, jump to a particular checkpoint" ) );
+static TAutoConsoleVariable<int32> CVarEnableCheckpoints( TEXT( "demo.EnableCheckpoints" ), 1, TEXT( "Whether or not checkpoints save on the server" ) );
+static TAutoConsoleVariable<int32> CVarGotoCheckpointIndex( TEXT( "demo.GotoCheckpointIndex" ), -2, TEXT( "For testing only, jump to a particular checkpoint" ) );
+static TAutoConsoleVariable<float> CVarGotoTimeInSeconds( TEXT( "demo.GotoTimeInSeconds" ), -1, TEXT( "For testing only, jump to a particular time" ) );
 static TAutoConsoleVariable<int32> CVarDemoFastForwardDestroyTearOffActors( TEXT( "demo.FastForwardDestroyTearOffActors" ), 1, TEXT( "If true, the driver will destroy any torn-off actors immediately while fast-forwarding a replay." ) );
 static TAutoConsoleVariable<int32> CVarDemoFastForwardSkipRepNotifies( TEXT( "demo.FastForwardSkipRepNotifies" ), 1, TEXT( "If true, the driver will optimize fast-forwarding by deferring calls to RepNotify functions until the fast-forward is complete. " ) );
 
@@ -55,8 +56,8 @@ bool UDemoNetDriver::InitBase( bool bInitAsClient, FNetworkNotify* InNotify, con
 		bChannelsArePaused		= false;
 		TimeToSkip				= 0.0f;
 		bIsFastForwarding		= false;
-		FastForwardStartSeconds	= 0.0;
-		CurrentCheckpointIndex	= 0;
+		GotoCheckpointSkipExtraTimeInMS = -1;
+		LastGotoTimeInSeconds	= -1.0f;
 
 		ResetDemoState();
 
@@ -500,17 +501,6 @@ void UDemoNetDriver::TickFlush( float DeltaSeconds )
 			}
 		}
 
-		if ( CVarTestCheckpoint.GetValueOnGameThread() >= 0 )
-		{
-			ReplayStreamer->GotoCheckpoint( CVarTestCheckpoint.GetValueOnGameThread(), FOnCheckpointReadyDelegate::CreateUObject( this, &UDemoNetDriver::CheckpointReady ) );
-			CVarTestCheckpoint.AsVariable()->Set( TEXT( "-1" ), ECVF_SetByConsole );
-		}
-
-		if ( bDemoPlaybackDone )
-		{
-			return;
-		}
-
 		if ( World->GetWorldSettings()->Pauser == NULL )
 		{
 			TickDemoPlayback( DeltaSeconds );
@@ -834,7 +824,7 @@ void UDemoNetDriver::SaveCheckpoint()
 
 	if ( CheckpointArchive->TotalSize() > 0 )
 	{
-		ReplayStreamer->FlushCheckpoint( CurrentCheckpointIndex++ );
+		ReplayStreamer->FlushCheckpoint( SavedAbsTimeMS );
 	}
 
 	const double EndCheckpointTime = FPlatformTime::Seconds();
@@ -929,7 +919,7 @@ void UDemoNetDriver::TickDemoRecord( float DeltaSeconds )
 	// Save a checkpoint if it's time
 	if ( CVarEnableCheckpoints.GetValueOnGameThread() == 1 )
 	{
-		const double CHECKPOINT_DELAY = 10.0;
+		const double CHECKPOINT_DELAY = 30.0;
 
 		if ( CurrentSeconds - LastCheckpointTime > CHECKPOINT_DELAY )
 		{
@@ -947,6 +937,7 @@ void UDemoNetDriver::PauseChannels( const bool bPause )
 	}
 
 	// Pause all non player controller actors
+	// FIXME: Would love a more elegant way of handling this at a more global level
 	for ( int32 i = ServerConnection->OpenChannels.Num() - 1; i >= 0; i-- )
 	{
 		UChannel* OpenChannel = ServerConnection->OpenChannels[i];
@@ -997,6 +988,7 @@ bool UDemoNetDriver::ConditionallyReadDemoFrame()
 		return false;
 	} 
 
+	// We're going to read a frame, unpause channels
 	PauseChannels( false );
 
 	const int32 OldFilePos = FileAr->Tell();
@@ -1116,11 +1108,16 @@ bool UDemoNetDriver::ReadDemoFrame( FArchive* Archive )
 	return true;
 }
 
-void UDemoNetDriver::SkipTime(float InTimeToSkip)
+void UDemoNetDriver::SkipTime(const float InTimeToSkip)
 {
 	TimeToSkip = InTimeToSkip;
 }
 
+void UDemoNetDriver::GotoTimeInSeconds( const float TimeInSeconds )
+{
+	LastGotoTimeInSeconds = TimeInSeconds;
+}
+
 void UDemoNetDriver::TickDemoPlayback( float DeltaSeconds )
 {
 	if ( ServerConnection == NULL || ServerConnection->State == USOCK_Closed )
@@ -1129,57 +1126,98 @@ void UDemoNetDriver::TickDemoPlayback( float DeltaSeconds )
 		return;
 	}
 
-	const uint32 TotalDemoTimeInMS = ReplayStreamer->GetTotalDemoTime();
-
-	if ( TotalDemoTimeInMS > 0 )
+	if ( GotoCheckpointArchive == NULL )
 	{
-		DemoTotalTime = (float)TotalDemoTimeInMS / 1000.0f;
+		// Check checkpoint debug cvars
+		check( GotoCheckpointSkipExtraTimeInMS == -1 );
+
+		if ( CVarGotoCheckpointIndex.GetValueOnGameThread() >= -1 )
+		{
+			ReplayStreamer->GotoCheckpointIndex( CVarGotoCheckpointIndex.GetValueOnGameThread(), FOnCheckpointReadyDelegate::CreateUObject( this, &UDemoNetDriver::CheckpointReady ) );
+			CVarGotoCheckpointIndex.AsVariable()->Set( TEXT( "-2" ), ECVF_SetByConsole );
+		}
+
+		if ( CVarGotoTimeInSeconds.GetValueOnGameThread() >= 0.0f )
+		{
+			ReplayStreamer->GotoTimeInMS( CVarGotoTimeInSeconds.GetValueOnGameThread() * 1000, FOnCheckpointReadyDelegate::CreateUObject( this, &UDemoNetDriver::CheckpointReady ) );
+			CVarGotoTimeInSeconds.AsVariable()->Set( TEXT( "-1" ), ECVF_SetByConsole );
+		}
+
+		if ( LastGotoTimeInSeconds >= 0.0f )
+		{
+			ReplayStreamer->GotoTimeInMS( LastGotoTimeInSeconds * 1000, FOnCheckpointReadyDelegate::CreateUObject( this, &UDemoNetDriver::CheckpointReady ) );
+			LastGotoTimeInSeconds = -1.0f;
+		}
 	}
 
 	if ( CVarDemoSkipTime.GetValueOnGameThread() > 0 )
 	{
-		SkipTime( CVarDemoSkipTime.GetValueOnGameThread() );		// Just overwrite existing value, cvar wins in this case
+		// Just overwrite existing value, cvar wins in this case
+		SkipTime( CVarDemoSkipTime.GetValueOnGameThread() );		
 		CVarDemoSkipTime.AsVariable()->Set( TEXT( "0" ), ECVF_SetByConsole );
 	}
-	
-	if ( TimeToSkip > 0.0f )
-	{
-		const uint32 TimeInMSToCheck = FMath::Clamp( GetDemoCurrentTimeInMS() + (uint32)( TimeToSkip * 1000 ), (uint32)0, TotalDemoTimeInMS );
 
-		ReplayStreamer->SetHighPriorityTimeRange( TimeInMSToCheck, TimeInMSToCheck );
+	// Update total demo time
+	const uint32 TotalDemoTimeInMS = ReplayStreamer->GetTotalDemoTime();
 
-		if ( !ReplayStreamer->IsDataAvailableForTimeRange( TimeInMSToCheck, TimeInMSToCheck ) )
-		{
-			PauseChannels( true );
-			return;
-		}
+	if ( TotalDemoTimeInMS > 0 )
+	{
+		DemoTotalTime = (float)TotalDemoTimeInMS / 1000.0f;
+	}
 
-		ReplayStreamer->SetHighPriorityTimeRange( 0, 0 );
+	// See if we have time to skip
+	if ( TimeToSkip > 0.0f )
+	{
+		const uint32 TimeInMSToCheck = FMath::Clamp( GetDemoCurrentTimeInMS() + (uint32)( TimeToSkip * 1000 ), (uint32)0, TotalDemoTimeInMS );
 
-		PauseChannels( false );
+		ReplayStreamer->SetHighPriorityTimeRange( GetDemoCurrentTimeInMS(), TimeInMSToCheck );
 
 		DemoCurrentTime += TimeToSkip;
 
+		bIsFastForwarding = true;
+
 		TimeToSkip = 0.0f;
+	}
 
-		bIsFastForwarding = true;
+	// Make sure there is data available to read
+	if ( !ReplayStreamer->IsDataAvailable() )
+	{
+		PauseChannels( true );
+		return;
+	}
 
-		FastForwardStartSeconds = FPlatformTime::Seconds();
+	// Check to see if we have a pending checkpoint to read from
+	if ( GotoCheckpointArchive != NULL )
+	{
+		LoadCheckpoint();
 	}
 
+	// If we're at the end of the demo, just pause channels and return
+	if ( bDemoPlaybackDone )
+	{
+		PauseChannels( true );
+		return;
+	}
+
+	// Advance demo time by seconds passed
 	DemoCurrentTime += DeltaSeconds;
 
+	// Clamp time
 	if ( DemoCurrentTime > DemoTotalTime )
 	{
 		DemoCurrentTime = DemoTotalTime;
 	}
 
-	// Read demo frames until we are caught up
+	// Speculatively grab seconds now in case we need it to get the time it took to fast forward
+	const double FastForwardStartSeconds = FPlatformTime::Seconds();
+
+	// Read demo frames until we are caught up (this implicitly handles fast forward if DemoCurrentTime past many frames)
 	while ( ConditionallyReadDemoFrame() )
 	{
 		DemoFrameNum++;
 	}
 
+	// Finalize any fast forward stuff that needs to happen
 	if ( bIsFastForwarding )
 	{
 		bIsFastForwarding = false;
@@ -1257,36 +1295,44 @@ void UDemoNetDriver::ReplayStreamingReady( bool bSuccess, bool bRecord )
 	}
 }
 
-void UDemoNetDriver::CheckpointReady( bool bSuccess )
+void UDemoNetDriver::CheckpointReady( const bool bSuccess, const int64 SkipExtraTimeInMS )
 {
+	check( GotoCheckpointArchive == NULL );
+	check( GotoCheckpointSkipExtraTimeInMS == -1 );
+
 	if ( !bSuccess )
 	{
 		UE_LOG( LogDemo, Warning, TEXT( "UDemoNetConnection::CheckpointReady: Failed to go to checkpoint." ) );
 		return;
 	}
 
-	FArchive* CheckpointArchive = ReplayStreamer->GetCheckpointArchive();
+	GotoCheckpointArchive			= ReplayStreamer->GetCheckpointArchive();
+	GotoCheckpointSkipExtraTimeInMS = SkipExtraTimeInMS;
+}
 
-	if ( CheckpointArchive->TotalSize() == 0 )
+void UDemoNetDriver::LoadCheckpoint()
+{
+	if ( GotoCheckpointArchive == NULL )
 	{
-		UE_LOG( LogDemo, Warning, TEXT( "UDemoNetConnection::CheckpointReady: Checkpoint is empty." ) );
+		UE_LOG( LogDemo, Warning, TEXT( "UDemoNetConnection::LoadCheckpoint: GotoCheckpointArchive == NULL." ) );
 		return;
 	}
 
 	if ( SpectatorController == NULL )
 	{
-		UE_LOG( LogDemo, Warning, TEXT( "UDemoNetConnection::CheckpointReady: No spectator player controller." ) );
+		UE_LOG( LogDemo, Warning, TEXT( "UDemoNetConnection::LoadCheckpoint: No spectator player controller." ) );
 		return;
 	}
 
 	if ( SpectatorController->GetSpectatorPawn() == NULL )
 	{
-		UE_LOG( LogDemo, Warning, TEXT( "UDemoNetConnection::CheckpointReady: No spectator pawn." ) );
+		UE_LOG( LogDemo, Warning, TEXT( "UDemoNetConnection::LoadCheckpoint: No spectator pawn." ) );
 		return;
 	}
 
-	const FVector OldLocation = SpectatorController->GetSpectatorPawn()->GetActorLocation();
-	const FRotator OldRotation = SpectatorController->GetControlRotation();//GetSpectatorPawn()->GetActorRotation();
+	bRestoreSpectatorPosition = true;
+	SpectatorLocation = SpectatorController->GetSpectatorPawn()->GetActorLocation();
+	SpectatorRotation = SpectatorController->GetControlRotation();//GetSpectatorPawn()->GetActorRotation();
 
 	FURL ConnectURL;
 	ConnectURL.Map = DemoFilename;
@@ -1329,27 +1375,43 @@ void UDemoNetDriver::CheckpointReady( bool bSuccess )
 	GuidCache->ObjectLookup.Empty();
 	GuidCache->NetGUIDLookup.Empty();
 
+	if ( GotoCheckpointArchive->TotalSize() == 0 )
+	{
+		// This is the very first checkpoint, we'll read the stream from the very beginning in this case
+		DemoCurrentTime			= 0;
+		bDemoPlaybackDone		= false;
+
+		if ( GotoCheckpointSkipExtraTimeInMS != -1 )
+		{
+			DemoCurrentTime += (float)GotoCheckpointSkipExtraTimeInMS / 1000;
+		}
+
+		GotoCheckpointArchive			= NULL;
+		GotoCheckpointSkipExtraTimeInMS = -1;
+		return;
+	}
+
 	int32 NumValues = 0;
-	*CheckpointArchive << NumValues;
+	*GotoCheckpointArchive << NumValues;
 
 	for ( int32 i = 0; i < NumValues; i++ )
 	{
 		FNetworkGUID Guid;
 		
-		*CheckpointArchive << Guid;
+		*GotoCheckpointArchive << Guid;
 		
 		FNetGuidCacheObject CacheObject;
 
 		FString PathName;
 
-		*CheckpointArchive << CacheObject.OuterGUID;
-		*CheckpointArchive << PathName;
-		*CheckpointArchive << CacheObject.PackageGuid;
+		*GotoCheckpointArchive << CacheObject.OuterGUID;
+		*GotoCheckpointArchive << PathName;
+		*GotoCheckpointArchive << CacheObject.PackageGuid;
 
 		CacheObject.PathName = FName( *PathName );
 
 		uint8 Flags = 0;
-		*CheckpointArchive << Flags;
+		*GotoCheckpointArchive << Flags;
 
 		CacheObject.bNoLoad = ( Flags & ( 1 << 0 ) ) ? true : false;
 		CacheObject.bIgnoreWhenMissing = ( Flags & ( 1 << 1 ) ) ? true : false;		
@@ -1359,22 +1421,27 @@ void UDemoNetDriver::CheckpointReady( bool bSuccess )
 	}
 
 	uint32 SavedAbsTimeMS = 0;
-	*CheckpointArchive << SavedAbsTimeMS;
+	*GotoCheckpointArchive << SavedAbsTimeMS;
 
 	DemoCurrentTime = (float)SavedAbsTimeMS / 1000.0f;
 
-	ReadDemoFrame( CheckpointArchive );
-
-	bDemoPlaybackDone = false;
+	ReadDemoFrame( GotoCheckpointArchive );
 
-	// Persist location and rotation from previous spectator
-	if ( SpectatorController != NULL )
+	// If we need to skip more time for fine scrubbing, set that up now
+	if ( GotoCheckpointSkipExtraTimeInMS != -1 )
 	{
-		SpectatorController->SetInitialLocationAndRotation( OldLocation, OldRotation );
+		DemoCurrentTime += (float)GotoCheckpointSkipExtraTimeInMS / 1000;
+		bIsFastForwarding = true;
 	}
-	else
+
+	GotoCheckpointSkipExtraTimeInMS = -1;
+	GotoCheckpointArchive			= NULL;
+	bDemoPlaybackDone				= false;
+
+	// Make sure we reset the bRestoreSpectatorPosition. If not, we didn't receive the spectator, which should be replicated in every frame
+	if ( bRestoreSpectatorPosition )
 	{
-		UE_LOG( LogDemo, Warning, TEXT( "UDemoNetConnection::CheckpointReady: No spectator controller after checkpoint." ) );
+		UE_LOG( LogDemo, Warning, TEXT( "UDemoNetConnection::LoadCheckpoint: Spectator wasn't restored." ) );
 	}
 }
 
@@ -1490,6 +1557,14 @@ void UDemoNetConnection::HandleClientPlayer( APlayerController* PC, UNetConnecti
 	// Assume this is our special spectator controller
 	GetDriver()->SpectatorController = PC;
 
+	// Persist location and rotation from previous spectator
+	if ( GetDriver()->bRestoreSpectatorPosition )
+	{
+		GetDriver()->SpectatorController->SetInitialLocationAndRotation( GetDriver()->SpectatorLocation, GetDriver()->SpectatorRotation );
+		GetDriver()->bRestoreSpectatorPosition = false;
+		return;
+	}
+
 	for ( FActorIterator It( Driver->World ); It; ++It)
 	{
 		if ( It->IsA( APlayerStart::StaticClass() ) )
diff --git a/Engine/Source/Runtime/NetworkReplayStreaming/HttpNetworkReplayStreaming/Private/HttpNetworkReplayStreaming.cpp b/Engine/Source/Runtime/NetworkReplayStreaming/HttpNetworkReplayStreaming/Private/HttpNetworkReplayStreaming.cpp
--- a/Engine/Source/Runtime/NetworkReplayStreaming/HttpNetworkReplayStreaming/Private/HttpNetworkReplayStreaming.cpp
+++ b/Engine/Source/Runtime/NetworkReplayStreaming/HttpNetworkReplayStreaming/Private/HttpNetworkReplayStreaming.cpp
@@ -1,7 +1,6 @@
 // Copyright 1998-2015 Epic Games, Inc. All Rights Reserved.
 
 #include "HttpNetworkReplayStreaming.h"
-#include "OnlineJsonSerializer.h"
 
 DEFINE_LOG_CATEGORY_STATIC( LogHttpReplay, Log, All );
 
@@ -105,6 +104,7 @@ FHttpNetworkReplayStreamer::FHttpNetworkReplayStreamer() :
 	StreamChunkIndex( 0 ), 
 	LastChunkTime( 0 ), 
 	LastRefreshViewerTime( 0 ),
+	LastRefreshCheckpointTime( 0 ),
 	StreamerState( EStreamerState::Idle ), 
 	bStopStreamingCalled( false ), 
 	bNeedToUploadHeader( false ),
@@ -114,7 +114,9 @@ FHttpNetworkReplayStreamer::FHttpNetworkReplayStreamer() :
 	StreamTimeRangeStart( 0 ),
 	StreamTimeRangeEnd( 0 ),
 	HighPriorityEndTime( 0 ),
-	StreamerLastError( ENetworkReplayError::None )
+	StreamerLastError( ENetworkReplayError::None ),
+	DownloadCheckpointIndex( -1 ),
+	LastGotoTimeInMS( -1 )
 {
 	// Initialize the server URL
 	GConfig->GetString( TEXT( "HttpNetworkReplayStreaming" ), TEXT( "ServerURL" ), ServerURL, GEngineIni );
@@ -181,9 +183,14 @@ void FHttpNetworkReplayStreamer::StartStreaming( const FString& StreamName, bool
 
 		HttpRequest->OnProcessRequestComplete().BindRaw( this, &FHttpNetworkReplayStreamer::HttpStartDownloadingFinished );
 	
+		// Add the request to start downloading
 		AddRequestToQueue( EQueuedHttpRequestType::StartDownloading, HttpRequest );
 
+		// Download the header (will add to queue)
 		DownloadHeader();
+
+		// Download the first set of checkpoints (will add to queue)
+		EnumerateCheckpoints();
 	}
 	else
 	{
@@ -362,32 +369,108 @@ void FHttpNetworkReplayStreamer::FlushCheckpoint( const uint32 TimeInMS )
 	
 	// Flush any existing stream, we need checkpoints to line up with the next chunk
 	FlushStream();
+
 	// Flush the checkpoint
 	FlushCheckpointInternal( TimeInMS );
 }
 
-void FHttpNetworkReplayStreamer::GotoCheckpoint( const uint32 TimeInMS, const FOnCheckpointReadyDelegate& Delegate )
+void FHttpNetworkReplayStreamer::GotoCheckpointIndex( const int32 CheckpointIndex, const FOnCheckpointReadyDelegate& Delegate )
 {
 	if ( GotoCheckpointDelegate.IsBound() )
 	{
 		// If we're currently going to a checkpoint now, ignore this request
-		UE_LOG( LogHttpReplay, Warning, TEXT( "FHttpNetworkReplayStreamer::GotoCheckpoint. Busy processing another checkpoint." ) );
+		UE_LOG( LogHttpReplay, Warning, TEXT( "FHttpNetworkReplayStreamer::GotoCheckpointIndex. Busy processing another checkpoint." ) );
+		Delegate.ExecuteIfBound( false, -1 );
+		return;
+	}
+
+	check( DownloadCheckpointIndex == -1 );
+
+	if ( CheckpointIndex == -1 )
+	{
+		// Make sure to reset the checkpoint archive (this is how we signify that the engine should start from the beginning of the steam (we don't need a checkpoint for that))
+		CheckpointArchive.Buffer.Empty();
+		CheckpointArchive.Pos = 0;
+
+		// Completely reset our stream (we're going to start downloading from the start of the checkpoint)
+		StreamArchive.Buffer.Empty();
+		StreamArchive.Pos				= 0;
+		StreamArchive.bAtEndOfReplay	= false;
+
+		// Reset our stream range
+		StreamTimeRangeStart	= 0;
+		StreamTimeRangeEnd		= 0;
+
+		StreamChunkIndex		= 0;
+
+		Delegate.ExecuteIfBound( true, LastGotoTimeInMS );
+		LastGotoTimeInMS = -1;
+		return;
+	}
+
+	if ( !CheckpointList.Checkpoints.IsValidIndex( CheckpointIndex ) )
+	{
+		UE_LOG( LogHttpReplay, Warning, TEXT( "FHttpNetworkReplayStreamer::GotoCheckpointIndex. Invalid checkpoint index." ) );
+		Delegate.ExecuteIfBound( false, -1 );
 		return;
 	}
 
 	TSharedRef<class IHttpRequest> HttpRequest = FHttpModule::Get().CreateRequest();
 
 	// Download the next stream chunk
-	HttpRequest->SetURL( FString::Printf( TEXT( "%sdownload?Version=%s&Session=%s&Filename=checkpoint.%i" ), *ServerURL, *SessionVersion, *SessionName, TimeInMS ) );
+	HttpRequest->SetURL( FString::Printf( TEXT( "%sdownloadevent?Version=%s&Session=%s&ID=%s" ), *ServerURL, *SessionVersion, *SessionName, *CheckpointList.Checkpoints[CheckpointIndex].ID ) );
 	HttpRequest->SetVerb( TEXT( "GET" ) );
 
 	HttpRequest->OnProcessRequestComplete().BindRaw( this, &FHttpNetworkReplayStreamer::HttpDownloadCheckpointFinished );
 
-	GotoCheckpointDelegate = Delegate;
+	GotoCheckpointDelegate	= Delegate;
+	DownloadCheckpointIndex = CheckpointIndex;
 
 	AddRequestToQueue( EQueuedHttpRequestType::DownloadingCheckpoint, HttpRequest );
 }
 
+void FHttpNetworkReplayStreamer::GotoTimeInMS( const uint32 TimeInMS, const FOnCheckpointReadyDelegate& Delegate )
+{
+	if ( GotoCheckpointDelegate.IsBound() )
+	{
+		// If we're currently going to a checkpoint now, ignore this request
+		UE_LOG( LogHttpReplay, Warning, TEXT( "FHttpNetworkReplayStreamer::GotoTimeInMS. Busy processing another checkpoint." ) );
+		Delegate.ExecuteIfBound( false, -1 );
+		return;
+	}
+
+	UE_LOG( LogHttpReplay, Verbose, TEXT( "FHttpNetworkReplayStreamer::GotoTimeInMS. TimeInMS: %i" ), (int)TimeInMS );
+
+	check( DownloadCheckpointIndex == -1 );
+	check( LastGotoTimeInMS == -1 );
+
+	int32 CheckpointIndex = -1;
+
+	LastGotoTimeInMS = TimeInMS;
+
+	if ( CheckpointList.Checkpoints.Num() > 0 && TimeInMS >= CheckpointList.Checkpoints[ CheckpointList.Checkpoints.Num() - 1 ].Time1 )
+	{
+		// If we're after the very last checkpoint, that's the one we want
+		CheckpointIndex = CheckpointList.Checkpoints.Num() - 1;
+	}
+	else
+	{
+		// Checkpoints should be sorted by time, return the checkpoint that exists right before the current time
+		// For fine scrubbing, we'll fast forward the rest of the way
+		// NOTE - If we're right before the very first checkpoint, we'll return -1, which is what we want when we want to start from the very beginning
+		for ( int i = 0; i < CheckpointList.Checkpoints.Num(); i++ )
+		{
+			if ( TimeInMS < CheckpointList.Checkpoints[i].Time1 )
+			{
+				CheckpointIndex = i - 1;
+				break;
+			}
+		}
+	}
+
+	GotoCheckpointIndex( CheckpointIndex, Delegate );
+}
+
 void FHttpNetworkReplayStreamer::FlushCheckpointInternal( uint32 TimeInMS )
 {
 	if ( SessionName.IsEmpty() || StreamerState != EStreamerState::StreamingUp || CheckpointArchive.Buffer.Num() == 0 )
@@ -406,10 +489,7 @@ void FHttpNetworkReplayStreamer::FlushCheckpointInternal( uint32 TimeInMS )
 
 	HttpRequest->OnProcessRequestComplete().BindRaw( this, &FHttpNetworkReplayStreamer::HttpUploadCheckpointFinished );
 
-	// Save the chunk to the checkpoint
-	CheckpointArchive << StreamChunkIndex;
-
-	HttpRequest->SetURL( FString::Printf( TEXT( "%supload?Version=%s&Session=%s&NumChunks=%i&Time=%i&Filename=checkpoint.%i" ), *ServerURL, *SessionVersion, *SessionName, StreamChunkIndex, TotalDemoTimeInMS, TimeInMS ) );
+	HttpRequest->SetURL( FString::Printf( TEXT( "%suploadevent?Version=%s&Session=%s&Group=checkpoint&Time1=%i&Time2=%i&Meta=%i" ), *ServerURL, *SessionVersion, *SessionName, TimeInMS, TimeInMS, StreamChunkIndex ) );
 	HttpRequest->SetVerb( TEXT( "POST" ) );
 	HttpRequest->SetHeader( TEXT( "Content-Type" ), TEXT( "application/octet-stream" ) );
 	HttpRequest->SetContent( CheckpointArchive.Buffer );
@@ -582,6 +662,14 @@ bool FHttpNetworkReplayStreamer::IsDataAvailable() const
 		return false;
 	}
 
+	if ( HighPriorityEndTime > 0 )
+	{
+		// If we are waiting for a high priority portion of the stream, pretend like we don't have any data so that game code waits for the entire portion
+		// of the high priority stream to download.
+		// We do this because we assume the game wants to race through this high priority portion of the stream in a single frame
+		return false;
+	}
+
 	// If we are loading, and we have more data
 	if ( StreamArchive.IsLoading() && StreamArchive.Pos < StreamArchive.Buffer.Num() && NumDownloadChunks > 0 )
 	{
@@ -593,6 +681,8 @@ bool FHttpNetworkReplayStreamer::IsDataAvailable() const
 
 void FHttpNetworkReplayStreamer::SetHighPriorityTimeRange( const uint32 StartTimeInMS, const uint32 EndTimeInMS )
 {
+	// Set the amount of stream we should download before saying we have anymore data available
+	// We will also put a high priority on this portion of the stream so that it downloads as fast as possible
 	HighPriorityEndTime = EndTimeInMS;
 }
 
@@ -639,6 +729,38 @@ void FHttpNetworkReplayStreamer::EnumerateStreams( const FString& VersionString,
 	AddRequestToQueue( EQueuedHttpRequestType::EnumeratingSessions, HttpRequest );
 }
 
+void FHttpNetworkReplayStreamer::EnumerateCheckpoints()
+{
+	TSharedRef<class IHttpRequest> HttpRequest = FHttpModule::Get().CreateRequest();
+
+	// Enumerate all of the sessions
+	HttpRequest->SetURL( FString::Printf( TEXT( "%senumevents?Version=%s&Session=%s&Group=checkpoint" ), *ServerURL, *SessionVersion, *SessionName ) );
+	HttpRequest->SetVerb( TEXT( "GET" ) );
+
+	HttpRequest->OnProcessRequestComplete().BindRaw( this, &FHttpNetworkReplayStreamer::HttpEnumerateCheckpointsFinished );
+
+	AddRequestToQueue( EQueuedHttpRequestType::EnumeratingCheckpoints, HttpRequest );
+
+	LastRefreshCheckpointTime = FPlatformTime::Seconds();
+}
+
+void FHttpNetworkReplayStreamer::ConditionallyEnumerateCheckpoints()
+{
+	if ( !bStreamIsLive )
+	{
+		// We don't need to enumerate more than once for non live streams
+		return;
+	}
+
+	// During live games, check for new checkpoints every 30 seconds
+	const double REFRESH_CHECKPOINTS_IN_SECONDS = 30;
+
+	if ( FPlatformTime::Seconds() - LastRefreshCheckpointTime > REFRESH_CHECKPOINTS_IN_SECONDS )
+	{
+		EnumerateCheckpoints();
+	}
+};
+
 void FHttpNetworkReplayStreamer::RequestFinished( EStreamerState ExpectedStreamerState, EQueuedHttpRequestType ExpectedType, FHttpRequestPtr HttpRequest )
 {
 	check( StreamerState == ExpectedStreamerState );
@@ -863,55 +985,64 @@ void FHttpNetworkReplayStreamer::HttpDownloadFinished( FHttpRequestPtr HttpReque
 
 void FHttpNetworkReplayStreamer::HttpDownloadCheckpointFinished( FHttpRequestPtr HttpRequest, FHttpResponsePtr HttpResponse, bool bSucceeded )
 {
+	UE_LOG( LogHttpReplay, Verbose, TEXT( "FHttpNetworkReplayStreamer::HttpDownloadCheckpointFinished." ) );
+
 	RequestFinished( EStreamerState::StreamingDown, EQueuedHttpRequestType::DownloadingCheckpoint, HttpRequest );
 
 	check( StreamArchive.IsLoading() );
 	check( GotoCheckpointDelegate.IsBound() );
+	check( DownloadCheckpointIndex >= 0 );
 
 	if ( bSucceeded && HttpResponse->GetResponseCode() == EHttpResponseCodes::Ok )
 	{
 		if ( HttpResponse->GetContent().Num() == 0 )
 		{
 			UE_LOG( LogHttpReplay, Warning, TEXT( "FHttpNetworkReplayStreamer::HttpDownloadCheckpointFinished. Checkpoint empty." ) );
-			GotoCheckpointDelegate.ExecuteIfBound( true );
+			GotoCheckpointDelegate.ExecuteIfBound( false, -1 );
 			return;
 		}
 
-		StreamArchive.Buffer.Empty();
-		StreamArchive.Pos = 0;
-		StreamArchive.bAtEndOfReplay = false;
-
-		CheckpointArchive.Buffer = HttpResponse->GetContent();
-		CheckpointArchive.Pos = 0;
-
-		StreamTimeRangeStart = 0;
-		StreamTimeRangeEnd = 0;
+		// Get the checkpoint data
+		CheckpointArchive.Buffer	= HttpResponse->GetContent();
+		CheckpointArchive.Pos		= 0;
 
-		// Seek to the end, where we added our internal meta data
-		CheckpointArchive.Seek( CheckpointArchive.Buffer.Num() - 4 );
-		
-		// Set the next chunk to be the first chunk that represents the start of the stream after the checkpoint
-		CheckpointArchive << StreamChunkIndex;
+		// Completely reset our stream (we're going to start downloading from the start of the checkpoint)
+		StreamArchive.Buffer.Empty();
+		StreamArchive.Pos				= 0;
+		StreamArchive.bAtEndOfReplay	= false;
 
-		// Make sure we read to the very end
-		check( CheckpointArchive.Pos == CheckpointArchive.Buffer.Num() );
+		// Reset our stream range
+		StreamTimeRangeStart	= 0;
+		StreamTimeRangeEnd		= 0;
 
-		// Remove the data we serialized from the end of the checkpoint
-		CheckpointArchive.Buffer.SetNum( CheckpointArchive.Buffer.Num() - 4 );
+		// Set the next chunk to be right after this checkpoint (which was stored in the metadata)
+		StreamChunkIndex = FCString::Atoi( *CheckpointList.Checkpoints[ DownloadCheckpointIndex ].Metadata );
 
-		CheckpointArchive.Pos = 0;
+		if ( LastGotoTimeInMS >= 0 )
+		{
+			// If we are fine scrubbing, make sure to wait on the part of the stream that is needed to do this in one frame
+			SetHighPriorityTimeRange( CheckpointList.Checkpoints[ DownloadCheckpointIndex ].Time1, LastGotoTimeInMS );
+			
+			// Subtract off checkpoint time so we pass in the leftover to the engine to fast forward through for the fine scrubbing part
+			LastGotoTimeInMS -= CheckpointList.Checkpoints[ DownloadCheckpointIndex ].Time1;
+		}
 
-		GotoCheckpointDelegate.ExecuteIfBound( true );
+		// Notify game code of success
+		GotoCheckpointDelegate.ExecuteIfBound( true, LastGotoTimeInMS );
 
 		UE_LOG( LogHttpReplay, Verbose, TEXT( "FHttpNetworkReplayStreamer::HttpDownloadCheckpointFinished. SUCCESS. StreamChunkIndex: %i" ), StreamChunkIndex );
 	}
 	else
 	{
+		// Oops, something went wrong, notify game code of failure
 		UE_LOG( LogHttpReplay, Error, TEXT( "FHttpNetworkReplayStreamer::HttpDownloadCheckpointFinished. FAILED." ) );
-		GotoCheckpointDelegate.ExecuteIfBound( false );
+		GotoCheckpointDelegate.ExecuteIfBound( false, -1 );
 	}
 
-	GotoCheckpointDelegate = FOnCheckpointReadyDelegate();
+	// Reset things
+	GotoCheckpointDelegate	= FOnCheckpointReadyDelegate();
+	DownloadCheckpointIndex = -1;
+	LastGotoTimeInMS		= -1;
 }
 
 void FHttpNetworkReplayStreamer::HttpRefreshViewerFinished( FHttpRequestPtr HttpRequest, FHttpResponsePtr HttpResponse, bool bSucceeded )
@@ -975,6 +1106,32 @@ void FHttpNetworkReplayStreamer::HttpEnumerateSessionsFinished( FHttpRequestPtr
 	EnumerateStreamsDelegate = FOnEnumerateStreamsComplete();
 }
 
+void FHttpNetworkReplayStreamer::HttpEnumerateCheckpointsFinished( FHttpRequestPtr HttpRequest, FHttpResponsePtr HttpResponse, bool bSucceeded )
+{
+	RequestFinished( EStreamerState::StreamingDown, EQueuedHttpRequestType::EnumeratingCheckpoints, HttpRequest );
+
+	if ( bSucceeded && HttpResponse->GetResponseCode() == EHttpResponseCodes::Ok )
+	{		
+		UE_LOG( LogHttpReplay, Verbose, TEXT( "FHttpNetworkReplayStreamer::HttpEnumerateCheckpointsFinished." ) );
+
+		FString JsonString = HttpResponse->GetContentAsString();
+
+		CheckpointList.Checkpoints.Empty();
+
+		if ( !CheckpointList.FromJson( JsonString ) )
+		{
+			UE_LOG( LogHttpReplay, Warning, TEXT( "FHttpNetworkReplayStreamer::HttpEnumerateCheckpointsFinished. FromJson FAILED" ) );
+			SetLastError( ENetworkReplayError::ServiceUnavailable );
+			return;
+		}
+	}
+	else
+	{
+		UE_LOG( LogHttpReplay, Warning, TEXT( "FHttpNetworkReplayStreamer::HttpEnumerateCheckpointsFinished. FAILED" ) );
+		SetLastError( ENetworkReplayError::ServiceUnavailable );
+	}
+}
+
 bool FHttpNetworkReplayStreamer::ProcessNextHttpRequest()
 {
 	if ( IsHttpRequestInFlight() )
@@ -1049,13 +1206,23 @@ void FHttpNetworkReplayStreamer::Tick( const float DeltaTime )
 	}
 	else if ( StreamerState == EStreamerState::StreamingDown )
 	{
+		// Check to see if we're done downloading the high priority portion of the strema
+		// If so, we can cancel the request
+		if ( HighPriorityEndTime > 0 && StreamTimeRangeEnd >= HighPriorityEndTime )
+		{
+			HighPriorityEndTime = 0;
+		}
+
+		// Check to see if we're at the end of non live streams
 		if ( StreamChunkIndex >= NumDownloadChunks && !bStreamIsLive )
 		{
+			// Make note of when we reach the end of non live stream
 			StreamArchive.bAtEndOfReplay = true;
 		}
 
 		ConditionallyRefreshViewer();
 		ConditionallyDownloadNextChunk();
+		ConditionallyEnumerateCheckpoints();
 	}
 }
 
diff --git a/Engine/Source/Runtime/NetworkReplayStreaming/HttpNetworkReplayStreaming/Private/HttpNetworkReplayStreaming.h b/Engine/Source/Runtime/NetworkReplayStreaming/HttpNetworkReplayStreaming/Private/HttpNetworkReplayStreaming.h
--- a/Engine/Source/Runtime/NetworkReplayStreaming/HttpNetworkReplayStreaming/Private/HttpNetworkReplayStreaming.h
+++ b/Engine/Source/Runtime/NetworkReplayStreaming/HttpNetworkReplayStreaming/Private/HttpNetworkReplayStreaming.h
@@ -4,6 +4,44 @@
 #include "NetworkReplayStreaming.h"
 #include "Http.h"
 #include "Runtime/Engine/Public/Tickable.h"
+#include "OnlineJsonSerializer.h"
+
+class FCheckpointListItem : public FOnlineJsonSerializable
+{
+public:
+	FCheckpointListItem() {}
+	virtual ~FCheckpointListItem() {}
+
+	FString		ID;
+	FString		Group;
+	FString		Metadata;
+	uint32		Time1;
+	uint32		Time2;
+
+	// FOnlineJsonSerializable
+	BEGIN_ONLINE_JSON_SERIALIZER
+		ONLINE_JSON_SERIALIZE( "id",			ID );
+		ONLINE_JSON_SERIALIZE( "group",			Group );
+		ONLINE_JSON_SERIALIZE( "meta",			Metadata );
+		ONLINE_JSON_SERIALIZE( "time1",			Time1 );
+		ONLINE_JSON_SERIALIZE( "time2",			Time2 );
+	END_ONLINE_JSON_SERIALIZER
+};
+
+class FCheckpointList : public FOnlineJsonSerializable
+{
+public:
+	FCheckpointList()
+	{}
+	virtual ~FCheckpointList() {}
+
+	TArray< FCheckpointListItem > Checkpoints;
+
+	// FOnlineJsonSerializable
+	BEGIN_ONLINE_JSON_SERIALIZER
+		ONLINE_JSON_SERIALIZE_ARRAY_SERIALIZABLE( "events", Checkpoints, FCheckpointListItem );
+	END_ONLINE_JSON_SERIALIZER
+};
 
 /**
  * Archive used to buffer stream over http
@@ -35,6 +73,7 @@ enum class EQueuedHttpRequestType
 	DownloadingStream,			// We are in the process of downloading the replay stream
 	RefreshingViewer,			// We are refreshing the server to let it know we're still viewing
 	EnumeratingSessions,		// We are in the process of downloading the available sessions
+	EnumeratingCheckpoints,		// We are in the process of downloading the available checkpoints
 	UploadingCheckpoint,		// We are uploading a checkpoint
 	DownloadingCheckpoint		// We are downloading a checkpoint
 };
@@ -63,7 +102,8 @@ public:
 	virtual FArchive*	GetStreamingArchive() override;
 	virtual FArchive*	GetCheckpointArchive() override;
 	virtual void		FlushCheckpoint( const uint32 TimeInMS ) override;
-	virtual void		GotoCheckpoint( const uint32 TimeInMS, const FOnCheckpointReadyDelegate& Delegate ) override;
+	virtual void		GotoCheckpointIndex( const int32 CheckpointIndex, const FOnCheckpointReadyDelegate& Delegate ) override;
+	virtual void		GotoTimeInMS( const uint32 TimeInMS, const FOnCheckpointReadyDelegate& Delegate ) override;
 	virtual FArchive*	GetMetadataArchive() override;
 	virtual void		UpdateTotalDemoTime( uint32 TimeInMS ) override;
 	virtual uint32		GetTotalDemoTime() const override { return TotalDemoTimeInMS; }
@@ -87,6 +127,8 @@ public:
 	void SetLastError( const ENetworkReplayError::Type InLastError );
 	void FlushCheckpointInternal( uint32 TimeInMS );
 	void AddRequestToQueue( const EQueuedHttpRequestType Type, TSharedPtr< class IHttpRequest >	Request );
+	void EnumerateCheckpoints();
+	void ConditionallyEnumerateCheckpoints();
 
 	/** EStreamerState - Overall state of the streamer */
 	enum class EStreamerState
@@ -110,6 +152,7 @@ public:
 	void HttpUploadStreamFinished( FHttpRequestPtr HttpRequest, FHttpResponsePtr HttpResponse, bool bSucceeded );
 	void HttpUploadCheckpointFinished( FHttpRequestPtr HttpRequest, FHttpResponsePtr HttpResponse, bool bSucceeded );
 	void HttpEnumerateSessionsFinished( FHttpRequestPtr HttpRequest, FHttpResponsePtr HttpResponse, bool bSucceeded );
+	void HttpEnumerateCheckpointsFinished( FHttpRequestPtr HttpRequest, FHttpResponsePtr HttpResponse, bool bSucceeded );
 
 	bool ProcessNextHttpRequest();
 	void Tick( const float DeltaTime );
@@ -126,6 +169,7 @@ public:
 	int32					StreamChunkIndex;		// Used as a counter to increment the stream.x extension count
 	double					LastChunkTime;			// The last time we uploaded/downloaded a chunk
 	double					LastRefreshViewerTime;	// The last time we refreshed ourselves as an active viewer
+	double					LastRefreshCheckpointTime;
 	EStreamerState			StreamerState;			// Overall state of the streamer
 	bool					bStopStreamingCalled;
 	bool					bNeedToUploadHeader;	// We're waiting on session name so we can upload header
@@ -143,6 +187,10 @@ public:
 	FOnStreamReadyDelegate			StartStreamingDelegate;		// Delegate passed in to StartStreaming
 	FOnEnumerateStreamsComplete		EnumerateStreamsDelegate;
 	FOnCheckpointReadyDelegate		GotoCheckpointDelegate;
+	int32							DownloadCheckpointIndex;
+	int64							LastGotoTimeInMS;
+
+	FCheckpointList					CheckpointList;
 
 	TQueue< TSharedPtr< FQueuedHttpRequest > >	QueuedHttpRequests;
 	TSharedPtr< FQueuedHttpRequest >			InFlightHttpRequest;
diff --git a/Engine/Source/Runtime/NetworkReplayStreaming/NetworkReplayStreaming/Public/NetworkReplayStreaming.h b/Engine/Source/Runtime/NetworkReplayStreaming/NetworkReplayStreaming/Public/NetworkReplayStreaming.h
--- a/Engine/Source/Runtime/NetworkReplayStreaming/NetworkReplayStreaming/Public/NetworkReplayStreaming.h
+++ b/Engine/Source/Runtime/NetworkReplayStreaming/NetworkReplayStreaming/Public/NetworkReplayStreaming.h
@@ -63,21 +63,21 @@ namespace ENetworkReplayError
  * @param bWasSuccessful Whether streaming was started.
  * @param bRecord Whether streaming is recording or not (vs playing)
  */
-DECLARE_DELEGATE_TwoParams( FOnStreamReadyDelegate, bool, bool );
+DECLARE_DELEGATE_TwoParams( FOnStreamReadyDelegate, const bool, const bool );
 
 /**
- * Delegate called when GotoCheckpoint() completes.
+ * Delegate called when GotoCheckpointIndex() completes.
  *
  * @param bWasSuccessful Whether streaming was started.
  */
-DECLARE_DELEGATE_OneParam( FOnCheckpointReadyDelegate, bool );
+DECLARE_DELEGATE_TwoParams( FOnCheckpointReadyDelegate, const bool, const int64 );
 
 /**
  * Delegate called when DeleteFinishedStream() completes.
  *
  * @param bWasSuccessful Whether the stream was deleted.
  */
-DECLARE_DELEGATE_OneParam( FOnDeleteFinishedStreamComplete, bool );
+DECLARE_DELEGATE_OneParam( FOnDeleteFinishedStreamComplete, const bool );
 
 /**
  * Delegate called when EnumerateStreams() completes.
@@ -98,7 +98,8 @@ public:
 	virtual FArchive* GetStreamingArchive() = 0;
 	virtual FArchive* GetCheckpointArchive() = 0;
 	virtual void FlushCheckpoint( const uint32 TimeInMS ) = 0;
-	virtual void GotoCheckpoint( const uint32 TimeInMS, const FOnCheckpointReadyDelegate& Delegate ) = 0;
+	virtual void GotoCheckpointIndex( const int32 CheckpointIndex, const FOnCheckpointReadyDelegate& Delegate ) = 0;
+	virtual void GotoTimeInMS( const uint32 TimeInMS, const FOnCheckpointReadyDelegate& Delegate ) = 0;
 	virtual FArchive* GetMetadataArchive() = 0;
 	virtual void UpdateTotalDemoTime( uint32 TimeInMS ) = 0;
 	virtual uint32 GetTotalDemoTime() const = 0;
diff --git a/Engine/Source/Runtime/NetworkReplayStreaming/NullNetworkReplayStreaming/Public/NullNetworkReplayStreaming.h b/Engine/Source/Runtime/NetworkReplayStreaming/NullNetworkReplayStreaming/Public/NullNetworkReplayStreaming.h
--- a/Engine/Source/Runtime/NetworkReplayStreaming/NullNetworkReplayStreaming/Public/NullNetworkReplayStreaming.h
+++ b/Engine/Source/Runtime/NetworkReplayStreaming/NullNetworkReplayStreaming/Public/NullNetworkReplayStreaming.h
@@ -22,7 +22,8 @@ public:
 	virtual FArchive* GetStreamingArchive() override;
 	virtual FArchive* GetCheckpointArchive() override { return NULL; }
 	virtual void FlushCheckpoint( const uint32 TimeInMS ) override { }
-	virtual void GotoCheckpoint( const uint32 TimeInMS, const FOnCheckpointReadyDelegate& Delegate ) override { };
+	virtual void GotoCheckpointIndex( const int32 TimeInMS, const FOnCheckpointReadyDelegate& Delegate ) override { }
+	virtual void GotoTimeInMS( const uint32 TimeInMS, const FOnCheckpointReadyDelegate& Delegate ) override { }
 	virtual FArchive* GetMetadataArchive() override;
 	virtual void UpdateTotalDemoTime( uint32 TimeInMS ) override { }
 	virtual uint32 GetTotalDemoTime() const override { return 0; }
