We now pass in app string, version and changelist to replay server.

diff --git a/Engine/Source/Runtime/Engine/Private/DemoNetDriver.cpp b/Engine/Source/Runtime/Engine/Private/DemoNetDriver.cpp
--- a/Engine/Source/Runtime/Engine/Private/DemoNetDriver.cpp
+++ b/Engine/Source/Runtime/Engine/Private/DemoNetDriver.cpp
@@ -188,9 +188,7 @@ bool UDemoNetDriver::InitConnect( FNetworkNotify* InNotify, const FURL& ConnectU
 	ServerConnection = NewObject<UNetConnection>(GetTransientPackage(), UDemoNetConnection::StaticClass());
 	ServerConnection->InitConnection( this, USOCK_Pending, ConnectURL, 1000000 );
 
-	FString VersionString = FString::Printf( TEXT( "%s_%u" ), FApp::GetGameName(), FNetworkVersion::GetLocalNetworkVersion() );
-
-	ReplayStreamer->StartStreaming( DemoFilename, false, VersionString, FOnStreamReadyDelegate::CreateUObject( this, &UDemoNetDriver::ReplayStreamingReady ) );
+	ReplayStreamer->StartStreaming( DemoFilename, false, FNetworkVersion::GetReplayVersion(), FOnStreamReadyDelegate::CreateUObject( this, &UDemoNetDriver::ReplayStreamingReady ) );
 
 	return true;
 }
@@ -362,9 +360,7 @@ bool UDemoNetDriver::InitListen( FNetworkNotify* InNotify, FURL& ListenURL, bool
 	Connection->InitSendBuffer();
 	ClientConnections.Add( Connection );
 
-	FString VersionString = FString::Printf( TEXT( "%s_%u" ), FApp::GetGameName(), FNetworkVersion::GetLocalNetworkVersion() );
-
-	ReplayStreamer->StartStreaming( DemoFilename, true, VersionString, FOnStreamReadyDelegate::CreateUObject( this, &UDemoNetDriver::ReplayStreamingReady ) );
+	ReplayStreamer->StartStreaming( DemoFilename, true, FNetworkVersion::GetReplayVersion(), FOnStreamReadyDelegate::CreateUObject( this, &UDemoNetDriver::ReplayStreamingReady ) );
 
 	FArchive* FileAr = ReplayStreamer->GetHeaderArchive();
 
diff --git a/Engine/Source/Runtime/Engine/Private/UnrealNetwork.cpp b/Engine/Source/Runtime/Engine/Private/UnrealNetwork.cpp
--- a/Engine/Source/Runtime/Engine/Private/UnrealNetwork.cpp
+++ b/Engine/Source/Runtime/Engine/Private/UnrealNetwork.cpp
@@ -1,6 +1,7 @@
 // Copyright 1998-2015 Epic Games, Inc. All Rights Reserved.
 
 #include "EnginePrivate.h"
+#include "NetworkReplayStreaming.h"
 #include "Net/UnrealNetwork.h"
 #include "GeneralProjectSettings.h"
 
@@ -64,4 +65,10 @@ bool FNetworkVersion::IsNetworkCompatible( const uint32 LocalNetworkVersion, con
 	return LocalNetworkVersion == RemoteNetworkVersion;
 }
 
+FNetworkReplayVersion FNetworkVersion::GetReplayVersion()
+{
+	return FNetworkReplayVersion( FApp::GetGameName(), GetLocalNetworkVersion(), GEngineVersion.GetChangelist() );
+}
+
+
 // ----------------------------------------------------------------
diff --git a/Engine/Source/Runtime/Engine/Public/Net/UnrealNetwork.h b/Engine/Source/Runtime/Engine/Public/Net/UnrealNetwork.h
--- a/Engine/Source/Runtime/Engine/Public/Net/UnrealNetwork.h
+++ b/Engine/Source/Runtime/Engine/Public/Net/UnrealNetwork.h
@@ -23,6 +23,8 @@ class	UChannelIterator;
 class	UNetConnection;
 class	UPendingNetGame;
 
+class	FNetworkReplayVersion;
+
 /*-----------------------------------------------------------------------------
 	Types.
 -----------------------------------------------------------------------------*/
@@ -64,6 +66,11 @@ struct ENGINE_API FNetworkVersion
 	 */
 	static bool IsNetworkCompatible( const uint32 LocalNetworkVersion, const uint32 RemoteNetworkVersion );
 
+	/**
+	 * Generates a spexial struct that contains information to send to replay server
+	 */
+	static FNetworkReplayVersion GetReplayVersion();
+
 	/**
 	 * Internal network protocol version.
 	 * By default, this value is incorporated into the result of GetLocalNetworkVersion().
diff --git a/Engine/Source/Runtime/NetworkReplayStreaming/HttpNetworkReplayStreaming/Private/HttpNetworkReplayStreaming.cpp b/Engine/Source/Runtime/NetworkReplayStreaming/HttpNetworkReplayStreaming/Private/HttpNetworkReplayStreaming.cpp
--- a/Engine/Source/Runtime/NetworkReplayStreaming/HttpNetworkReplayStreaming/Private/HttpNetworkReplayStreaming.cpp
+++ b/Engine/Source/Runtime/NetworkReplayStreaming/HttpNetworkReplayStreaming/Private/HttpNetworkReplayStreaming.cpp
@@ -10,7 +10,7 @@ public:
 	FNetworkReplayListItem() {}
 	virtual ~FNetworkReplayListItem() {}
 
-	FString		VersionName;
+	FString		AppName;
 	FString		SessionName;
 	FString		FriendlyName;
 	FDateTime	Timestamp;
@@ -21,7 +21,7 @@ public:
 
 	// FOnlineJsonSerializable
 	BEGIN_ONLINE_JSON_SERIALIZER
-		ONLINE_JSON_SERIALIZE( "VersionName",	VersionName );
+		ONLINE_JSON_SERIALIZE( "AppName",		AppName );
 		ONLINE_JSON_SERIALIZE( "SessionName",	SessionName );
 		ONLINE_JSON_SERIALIZE( "FriendlyName",	FriendlyName );
 		ONLINE_JSON_SERIALIZE( "Timestamp",		Timestamp );
@@ -109,7 +109,7 @@ FHttpNetworkReplayStreamer::FHttpNetworkReplayStreamer() :
 	bStopStreamingCalled( false ), 
 	bNeedToUploadHeader( false ),
 	bStreamIsLive( false ),
-	NumDownloadChunks( 0 ),
+	NumTotalStreamChunks( 0 ),
 	TotalDemoTimeInMS( 0 ),
 	StreamTimeRangeStart( 0 ),
 	StreamTimeRangeEnd( 0 ),
@@ -122,7 +122,7 @@ FHttpNetworkReplayStreamer::FHttpNetworkReplayStreamer() :
 	GConfig->GetString( TEXT( "HttpNetworkReplayStreaming" ), TEXT( "ServerURL" ), ServerURL, GEngineIni );
 }
 
-void FHttpNetworkReplayStreamer::StartStreaming( const FString& StreamName, bool bRecord, const FString& VersionString, const FOnStreamReadyDelegate& Delegate )
+void FHttpNetworkReplayStreamer::StartStreaming( const FString& StreamName, bool bRecord, const FNetworkReplayVersion& InReplayVersion, const FOnStreamReadyDelegate& Delegate )
 {
 	if ( !SessionName.IsEmpty() )
 	{
@@ -142,7 +142,7 @@ void FHttpNetworkReplayStreamer::StartStreaming( const FString& StreamName, bool
 		return;
 	}
 
-	SessionVersion = VersionString;
+	ReplayVersion = InReplayVersion;
 
 	// Remember the delegate, which we'll call as soon as the header is available
 	StartStreamingDelegate = Delegate;
@@ -178,7 +178,7 @@ void FHttpNetworkReplayStreamer::StartStreaming( const FString& StreamName, bool
 		SessionName = StreamName;
 
 		// Notify the http server that we want to start downloading a replay
-		HttpRequest->SetURL( FString::Printf( TEXT( "%sstartdownloading?Version=%s&Session=%s" ), *ServerURL, *SessionVersion, *SessionName ) );
+		HttpRequest->SetURL( FString::Printf( TEXT( "%sstartdownloading?Session=%s" ), *ServerURL, *SessionName ) );
 		HttpRequest->SetVerb( TEXT( "POST" ) );
 
 		HttpRequest->OnProcessRequestComplete().BindRaw( this, &FHttpNetworkReplayStreamer::HttpStartDownloadingFinished );
@@ -203,7 +203,7 @@ void FHttpNetworkReplayStreamer::StartStreaming( const FString& StreamName, bool
 		bNeedToUploadHeader = true;
 
 		// Notify the http server that we want to start uploading a replay
-		HttpRequest->SetURL( FString::Printf( TEXT( "%sstartuploading?Version=%s&Friendly=%s" ), *ServerURL, *SessionVersion, *StreamName ) );
+		HttpRequest->SetURL( FString::Printf( TEXT( "%sstartuploading?App=%s&Version=%u&CL=%u&Friendly=%s" ), *ServerURL, *ReplayVersion.AppString, ReplayVersion.NetworkVersion, ReplayVersion.Changelist, *StreamName ) );
 
 		HttpRequest->SetVerb( TEXT( "POST" ) );
 
@@ -278,7 +278,7 @@ void FHttpNetworkReplayStreamer::UploadHeader()
 
 	HttpRequest->OnProcessRequestComplete().BindRaw( this, &FHttpNetworkReplayStreamer::HttpHeaderUploadFinished );
 
-	HttpRequest->SetURL( FString::Printf( TEXT( "%supload?Version=%s&Session=%s&NumChunks=%i&Time=%i&Filename=replay.header" ), *ServerURL, *SessionVersion, *SessionName, StreamChunkIndex, TotalDemoTimeInMS ) );
+	HttpRequest->SetURL( FString::Printf( TEXT( "%supload?Session=%s&NumChunks=%i&Time=%i&Filename=replay.header" ), *ServerURL, *SessionName, StreamChunkIndex, TotalDemoTimeInMS ) );
 	HttpRequest->SetVerb( TEXT( "POST" ) );
 	HttpRequest->SetHeader( TEXT( "Content-Type" ), TEXT( "application/octet-stream" ) );
 	HttpRequest->SetContent( HeaderArchive.Buffer );
@@ -317,7 +317,7 @@ void FHttpNetworkReplayStreamer::FlushStream()
 
 	HttpRequest->OnProcessRequestComplete().BindRaw( this, &FHttpNetworkReplayStreamer::HttpUploadStreamFinished );
 
-	HttpRequest->SetURL( FString::Printf( TEXT( "%supload?Version=%s&Session=%s&NumChunks=%i&Time=%i&MTime1=%i&MTime2=%i&Filename=stream.%i" ), *ServerURL, *SessionVersion, *SessionName, StreamChunkIndex + 1, TotalDemoTimeInMS, StreamTimeRangeStart, StreamTimeRangeEnd, StreamChunkIndex ) );
+	HttpRequest->SetURL( FString::Printf( TEXT( "%supload?Session=%s&NumChunks=%i&Time=%i&MTime1=%i&MTime2=%i&Filename=stream.%i" ), *ServerURL, *SessionName, StreamChunkIndex + 1, TotalDemoTimeInMS, StreamTimeRangeStart, StreamTimeRangeEnd, StreamChunkIndex ) );
 	HttpRequest->SetVerb( TEXT( "POST" ) );
 	HttpRequest->SetHeader( TEXT( "Content-Type" ), TEXT( "application/octet-stream" ) );
 	HttpRequest->SetContent( StreamArchive.Buffer );
@@ -352,7 +352,7 @@ void FHttpNetworkReplayStreamer::StopUploading()
 
 	HttpRequest->OnProcessRequestComplete().BindRaw( this, &FHttpNetworkReplayStreamer::HttpStopUploadingFinished );
 
-	HttpRequest->SetURL( FString::Printf( TEXT( "%sstopuploading?Version=%s&Session=%s&NumChunks=%i&Time=%i" ), *ServerURL, *SessionVersion, *SessionName, StreamChunkIndex, TotalDemoTimeInMS ) );
+	HttpRequest->SetURL( FString::Printf( TEXT( "%sstopuploading?Session=%s&NumChunks=%i&Time=%i" ), *ServerURL, *SessionName, StreamChunkIndex, TotalDemoTimeInMS ) );
 	HttpRequest->SetVerb( TEXT( "POST" ) );
 	HttpRequest->SetHeader( TEXT( "Content-Type" ), TEXT( "application/octet-stream" ) );
 
@@ -418,7 +418,7 @@ void FHttpNetworkReplayStreamer::GotoCheckpointIndex( const int32 CheckpointInde
 	TSharedRef<class IHttpRequest> HttpRequest = FHttpModule::Get().CreateRequest();
 
 	// Download the next stream chunk
-	HttpRequest->SetURL( FString::Printf( TEXT( "%sdownloadevent?Version=%s&Session=%s&ID=%s" ), *ServerURL, *SessionVersion, *SessionName, *CheckpointList.Checkpoints[CheckpointIndex].ID ) );
+	HttpRequest->SetURL( FString::Printf( TEXT( "%sdownloadevent?Session=%s&ID=%s" ), *ServerURL, *SessionName, *CheckpointList.Checkpoints[CheckpointIndex].ID ) );
 	HttpRequest->SetVerb( TEXT( "GET" ) );
 
 	HttpRequest->OnProcessRequestComplete().BindRaw( this, &FHttpNetworkReplayStreamer::HttpDownloadCheckpointFinished );
@@ -489,7 +489,7 @@ void FHttpNetworkReplayStreamer::FlushCheckpointInternal( uint32 TimeInMS )
 
 	HttpRequest->OnProcessRequestComplete().BindRaw( this, &FHttpNetworkReplayStreamer::HttpUploadCheckpointFinished );
 
-	HttpRequest->SetURL( FString::Printf( TEXT( "%suploadevent?Version=%s&Session=%s&Group=checkpoint&Time1=%i&Time2=%i&Meta=%i" ), *ServerURL, *SessionVersion, *SessionName, TimeInMS, TimeInMS, StreamChunkIndex ) );
+	HttpRequest->SetURL( FString::Printf( TEXT( "%suploadevent?Session=%s&Group=checkpoint&Time1=%i&Time2=%i&Meta=%i" ), *ServerURL, *SessionName, TimeInMS, TimeInMS, StreamChunkIndex ) );
 	HttpRequest->SetVerb( TEXT( "POST" ) );
 	HttpRequest->SetHeader( TEXT( "Content-Type" ), TEXT( "application/octet-stream" ) );
 	HttpRequest->SetContent( CheckpointArchive.Buffer );
@@ -505,7 +505,7 @@ void FHttpNetworkReplayStreamer::DownloadHeader()
 	// Download header first
 	TSharedRef<class IHttpRequest> HttpRequest = FHttpModule::Get().CreateRequest();
 
-	HttpRequest->SetURL( FString::Printf( TEXT( "%sdownload?Version=%s&Session=%s&Filename=replay.header" ), *ServerURL, *SessionVersion, *SessionName, *ViewerName ) );
+	HttpRequest->SetURL( FString::Printf( TEXT( "%sdownload?Session=%s&Filename=replay.header" ), *ServerURL, *SessionName, *ViewerName ) );
 	HttpRequest->SetVerb( TEXT( "GET" ) );
 
 	HttpRequest->OnProcessRequestComplete().BindRaw( this, &FHttpNetworkReplayStreamer::HttpDownloadHeaderFinished );
@@ -527,8 +527,8 @@ void FHttpNetworkReplayStreamer::ConditionallyDownloadNextChunk()
 		return;
 	}
 
-	const bool bMoreChunksDefinitelyAvilable	= StreamChunkIndex < NumDownloadChunks;		// We know for a fact there are more chunks available
-	const bool bMoreChunksMayBeAvilable			= bStreamIsLive;							// If we are live, there is a chance of more chunks even if we're at the end currently
+	const bool bMoreChunksDefinitelyAvilable	= StreamChunkIndex < NumTotalStreamChunks;		// We know for a fact there are more chunks available
+	const bool bMoreChunksMayBeAvilable			= bStreamIsLive;								// If we are live, there is a chance of more chunks even if we're at the end currently
 
 	if ( !bMoreChunksDefinitelyAvilable && !bMoreChunksMayBeAvilable )
 	{
@@ -557,7 +557,7 @@ void FHttpNetworkReplayStreamer::ConditionallyDownloadNextChunk()
 	TSharedRef<class IHttpRequest> HttpRequest = FHttpModule::Get().CreateRequest();
 
 	// Download the next stream chunk
-	HttpRequest->SetURL( FString::Printf( TEXT( "%sdownload?Version=%s&Session=%s&Filename=stream.%i" ), *ServerURL, *SessionVersion, *SessionName, StreamChunkIndex ) );
+	HttpRequest->SetURL( FString::Printf( TEXT( "%sdownload?Session=%s&Filename=stream.%i" ), *ServerURL, *SessionName, StreamChunkIndex ) );
 	HttpRequest->SetVerb( TEXT( "GET" ) );
 
 	HttpRequest->OnProcessRequestComplete().BindRaw( this, &FHttpNetworkReplayStreamer::HttpDownloadFinished );
@@ -574,11 +574,11 @@ void FHttpNetworkReplayStreamer::RefreshViewer( const bool bFinal )
 	// Download the next stream chunk
 	if ( bFinal )
 	{
-		HttpRequest->SetURL( FString::Printf( TEXT( "%srefreshviewer?Version=%s&Session=%s&Viewer=%s&Final=true" ), *ServerURL, *SessionVersion, *SessionName, *ViewerName ) );
+		HttpRequest->SetURL( FString::Printf( TEXT( "%srefreshviewer?Session=%s&Viewer=%s&Final=true" ), *ServerURL, *SessionName, *ViewerName ) );
 	}
 	else
 	{
-		HttpRequest->SetURL( FString::Printf( TEXT( "%srefreshviewer?Version=%s&Session=%s&Viewer=%s" ), *ServerURL, *SessionVersion, *SessionName, *ViewerName ) );
+		HttpRequest->SetURL( FString::Printf( TEXT( "%srefreshviewer?Session=%s&Viewer=%s" ), *ServerURL, *SessionName, *ViewerName ) );
 	}
 
 	HttpRequest->SetVerb( TEXT( "POST" ) );
@@ -671,7 +671,7 @@ bool FHttpNetworkReplayStreamer::IsDataAvailable() const
 	}
 
 	// If we are loading, and we have more data
-	if ( StreamArchive.IsLoading() && StreamArchive.Pos < StreamArchive.Buffer.Num() && NumDownloadChunks > 0 )
+	if ( StreamArchive.IsLoading() && StreamArchive.Pos < StreamArchive.Buffer.Num() && NumTotalStreamChunks > 0 )
 	{
 		return true;
 	}
@@ -708,7 +708,7 @@ void FHttpNetworkReplayStreamer::DeleteFinishedStream( const FString& StreamName
 	Delegate.ExecuteIfBound(false);
 }
 
-void FHttpNetworkReplayStreamer::EnumerateStreams( const FString& VersionString, const FOnEnumerateStreamsComplete& Delegate )
+void FHttpNetworkReplayStreamer::EnumerateStreams( const FNetworkReplayVersion& ReplayVersion, const FOnEnumerateStreamsComplete& Delegate )
 {
 	if ( IsHttpRequestInFlight() )
 	{
@@ -719,7 +719,7 @@ void FHttpNetworkReplayStreamer::EnumerateStreams( const FString& VersionString,
 	TSharedRef<class IHttpRequest> HttpRequest = FHttpModule::Get().CreateRequest();
 
 	// Enumerate all of the sessions
-	HttpRequest->SetURL( FString::Printf( TEXT( "%senumsessions?Version=%s" ), *ServerURL, *VersionString ) );
+	HttpRequest->SetURL( FString::Printf( TEXT( "%senumsessions?App=%s&Version=%u&CL=%u" ), *ServerURL, *ReplayVersion.AppString, ReplayVersion.NetworkVersion, ReplayVersion.Changelist ) );
 	HttpRequest->SetVerb( TEXT( "GET" ) );
 
 	HttpRequest->OnProcessRequestComplete().BindRaw( this, &FHttpNetworkReplayStreamer::HttpEnumerateSessionsFinished );
@@ -734,7 +734,7 @@ void FHttpNetworkReplayStreamer::EnumerateCheckpoints()
 	TSharedRef<class IHttpRequest> HttpRequest = FHttpModule::Get().CreateRequest();
 
 	// Enumerate all of the sessions
-	HttpRequest->SetURL( FString::Printf( TEXT( "%senumevents?Version=%s&Session=%s&Group=checkpoint" ), *ServerURL, *SessionVersion, *SessionName ) );
+	HttpRequest->SetURL( FString::Printf( TEXT( "%senumevents?Session=%s&Group=checkpoint" ), *ServerURL, *SessionName ) );
 	HttpRequest->SetVerb( TEXT( "GET" ) );
 
 	HttpRequest->OnProcessRequestComplete().BindRaw( this, &FHttpNetworkReplayStreamer::HttpEnumerateCheckpointsFinished );
@@ -877,12 +877,12 @@ void FHttpNetworkReplayStreamer::HttpStartDownloadingFinished( FHttpRequestPtr H
 
 		bStreamIsLive = State == TEXT( "Live" );
 
-		NumDownloadChunks = FCString::Atoi( *NumChunksString );
-		TotalDemoTimeInMS = FCString::Atoi( *DemoTimeString );
+		NumTotalStreamChunks	= FCString::Atoi( *NumChunksString );
+		TotalDemoTimeInMS		= FCString::Atoi( *DemoTimeString );
 
-		UE_LOG( LogHttpReplay, Log, TEXT( "FHttpNetworkReplayStreamer::HttpStartDownloadingFinished. Viewer: %s, State: %s, NumChunks: %i, DemoTime: %2.2f" ), *ViewerName, *State, NumDownloadChunks, (float)TotalDemoTimeInMS / 1000 );
+		UE_LOG( LogHttpReplay, Log, TEXT( "FHttpNetworkReplayStreamer::HttpStartDownloadingFinished. Viewer: %s, State: %s, NumChunks: %i, DemoTime: %2.2f" ), *ViewerName, *State, NumTotalStreamChunks, (float)TotalDemoTimeInMS / 1000 );
 
-		if ( NumDownloadChunks == 0 )
+		if ( NumTotalStreamChunks == 0 )
 		{			
 			UE_LOG( LogHttpReplay, Warning, TEXT( "FHttpNetworkReplayStreamer::HttpStartDownloadingFinished. NO CHUNKS" ) );
 
@@ -945,7 +945,7 @@ void FHttpNetworkReplayStreamer::HttpDownloadFinished( FHttpRequestPtr HttpReque
 
 	if ( bSucceeded && HttpResponse->GetResponseCode() == EHttpResponseCodes::Ok )
 	{
-		NumDownloadChunks = FCString::Atoi( *HttpResponse->GetHeader( TEXT( "NumChunks" ) ) );
+		NumTotalStreamChunks = FCString::Atoi( *HttpResponse->GetHeader( TEXT( "NumChunks" ) ) );
 		TotalDemoTimeInMS = FCString::Atoi( *HttpResponse->GetHeader( TEXT( "Time" ) ) );
 		bStreamIsLive = HttpResponse->GetHeader( TEXT( "State" ) ) == TEXT( "Live" );
 
@@ -966,7 +966,7 @@ void FHttpNetworkReplayStreamer::HttpDownloadFinished( FHttpRequestPtr HttpReque
 				StreamChunkIndex++;
 			}
 
-			UE_LOG( LogHttpReplay, Verbose, TEXT( "FHttpNetworkReplayStreamer::HttpDownloadFinished. Live: %i, Progress: %i / %i, Start: %i, End: %i, DemoTime: %2.2f" ), (int)bStreamIsLive, StreamChunkIndex, NumDownloadChunks, (int)StreamTimeRangeStart, (int)StreamTimeRangeEnd, (float)TotalDemoTimeInMS / 1000 );
+			UE_LOG( LogHttpReplay, Verbose, TEXT( "FHttpNetworkReplayStreamer::HttpDownloadFinished. Live: %i, Progress: %i / %i, Start: %i, End: %i, DemoTime: %2.2f" ), (int)bStreamIsLive, StreamChunkIndex, NumTotalStreamChunks, (int)StreamTimeRangeStart, (int)StreamTimeRangeEnd, (float)TotalDemoTimeInMS / 1000 );
 		}
 		else
 		{
@@ -1214,7 +1214,7 @@ void FHttpNetworkReplayStreamer::Tick( const float DeltaTime )
 		}
 
 		// Check to see if we're at the end of non live streams
-		if ( StreamChunkIndex >= NumDownloadChunks && !bStreamIsLive )
+		if ( StreamChunkIndex >= NumTotalStreamChunks && !bStreamIsLive )
 		{
 			// Make note of when we reach the end of non live stream
 			StreamArchive.bAtEndOfReplay = true;
diff --git a/Engine/Source/Runtime/NetworkReplayStreaming/HttpNetworkReplayStreaming/Private/HttpNetworkReplayStreaming.h b/Engine/Source/Runtime/NetworkReplayStreaming/HttpNetworkReplayStreaming/Private/HttpNetworkReplayStreaming.h
--- a/Engine/Source/Runtime/NetworkReplayStreaming/HttpNetworkReplayStreaming/Private/HttpNetworkReplayStreaming.h
+++ b/Engine/Source/Runtime/NetworkReplayStreaming/HttpNetworkReplayStreaming/Private/HttpNetworkReplayStreaming.h
@@ -96,7 +96,7 @@ public:
 	FHttpNetworkReplayStreamer();
 
 	/** INetworkReplayStreamer implementation */
-	virtual void		StartStreaming( const FString& StreamName, bool bRecord, const FString& VersionString, const FOnStreamReadyDelegate& Delegate ) override;
+	virtual void		StartStreaming( const FString& StreamName, bool bRecord, const FNetworkReplayVersion& ReplayVersion, const FOnStreamReadyDelegate& Delegate ) override;
 	virtual void		StopStreaming() override;
 	virtual FArchive*	GetHeaderArchive() override;
 	virtual FArchive*	GetStreamingArchive() override;
@@ -112,7 +112,7 @@ public:
 	virtual bool		IsDataAvailableForTimeRange( const uint32 StartTimeInMS, const uint32 EndTimeInMS ) override;
 	virtual bool		IsLive( const FString& StreamName ) const override;
 	virtual void		DeleteFinishedStream( const FString& StreamName, const FOnDeleteFinishedStreamComplete& Delegate ) const override;
-	virtual void		EnumerateStreams( const FString& VersionString, const FOnEnumerateStreamsComplete& Delegate ) override;
+	virtual void		EnumerateStreams( const FNetworkReplayVersion& ReplayVersion, const FOnEnumerateStreamsComplete& Delegate ) override;
 	virtual ENetworkReplayError::Type GetLastError() const override;
 
 	/** FHttpNetworkReplayStreamer */
@@ -164,7 +164,7 @@ public:
 	FHttpStreamFArchive		StreamArchive;			// Archive used to buffer the data stream
 	FHttpStreamFArchive		CheckpointArchive;		// Archive used to buffer checkpoint data
 	FString					SessionName;			// Name of the session on the http replay server
-	FString					SessionVersion;			// Version of the session
+	FNetworkReplayVersion	ReplayVersion;			// Version of the session
 	FString					ServerURL;				// The address of the server
 	int32					StreamChunkIndex;		// Used as a counter to increment the stream.x extension count
 	double					LastChunkTime;			// The last time we uploaded/downloaded a chunk
@@ -174,7 +174,7 @@ public:
 	bool					bStopStreamingCalled;
 	bool					bNeedToUploadHeader;	// We're waiting on session name so we can upload header
 	bool					bStreamIsLive;			// If true, we are viewing a live stream
-	int32					NumDownloadChunks;
+	int32					NumTotalStreamChunks;
 	uint32					TotalDemoTimeInMS;
 	uint32					LastTotalDemoTimeInMS;
 	uint32					StreamTimeRangeStart;
diff --git a/Engine/Source/Runtime/NetworkReplayStreaming/NetworkReplayStreaming/Public/NetworkReplayStreaming.h b/Engine/Source/Runtime/NetworkReplayStreaming/NetworkReplayStreaming/Public/NetworkReplayStreaming.h
--- a/Engine/Source/Runtime/NetworkReplayStreaming/NetworkReplayStreaming/Public/NetworkReplayStreaming.h
+++ b/Engine/Source/Runtime/NetworkReplayStreaming/NetworkReplayStreaming/Public/NetworkReplayStreaming.h
@@ -86,13 +86,24 @@ DECLARE_DELEGATE_OneParam( FOnDeleteFinishedStreamComplete, const bool );
  */
 DECLARE_DELEGATE_OneParam( FOnEnumerateStreamsComplete, const TArray<FNetworkReplayStreamInfo>& );
 
+class FNetworkReplayVersion
+{
+public:
+	FNetworkReplayVersion() : NetworkVersion( 0 ), Changelist( 0 ) {}
+	FNetworkReplayVersion( const FString& InAppString, const uint32 InNetworkVersion, const uint32 InChangelist ) : AppString( InAppString ), NetworkVersion( InNetworkVersion ), Changelist( InChangelist ) {}
+
+	FString		AppString;
+	uint32		NetworkVersion;
+	uint32		Changelist;
+};
+
 /** Generic interface for network replay streaming */
 class INetworkReplayStreamer 
 {
 public:
 	virtual ~INetworkReplayStreamer() {}
 
-	virtual void StartStreaming( const FString& StreamName, bool bRecord, const FString& VersionString, const FOnStreamReadyDelegate& Delegate ) = 0;
+	virtual void StartStreaming( const FString& StreamName, bool bRecord, const FNetworkReplayVersion& ReplayVersion, const FOnStreamReadyDelegate& Delegate ) = 0;
 	virtual void StopStreaming() = 0;
 	virtual FArchive* GetHeaderArchive() = 0;
 	virtual FArchive* GetStreamingArchive() = 0;
@@ -123,7 +134,7 @@ public:
 	 *
 	 * @param Delegate A delegate that will be executed if bound when the list of streams is available
 	 */
-	virtual void EnumerateStreams( const FString& VersionString, const FOnEnumerateStreamsComplete& Delegate ) = 0;
+	virtual void EnumerateStreams( const FNetworkReplayVersion& ReplayVersion, const FOnEnumerateStreamsComplete& Delegate ) = 0;
 
 	/** Returns the last error that occurred while streaming replays */
 	virtual ENetworkReplayError::Type GetLastError() const = 0;
diff --git a/Engine/Source/Runtime/NetworkReplayStreaming/NullNetworkReplayStreaming/Private/NullNetworkReplayStreaming.cpp b/Engine/Source/Runtime/NetworkReplayStreaming/NullNetworkReplayStreaming/Private/NullNetworkReplayStreaming.cpp
--- a/Engine/Source/Runtime/NetworkReplayStreaming/NullNetworkReplayStreaming/Private/NullNetworkReplayStreaming.cpp
+++ b/Engine/Source/Runtime/NetworkReplayStreaming/NullNetworkReplayStreaming/Private/NullNetworkReplayStreaming.cpp
@@ -65,7 +65,7 @@ static FString GetMetadataFilename(const FString& StreamName)
 	return GetStreamFullBaseFilename(StreamName) + TEXT(".metadata");
 }
 
-void FNullNetworkReplayStreamer::StartStreaming( const FString& StreamName, bool bRecord, const FString& VersionString, const FOnStreamReadyDelegate& Delegate )
+void FNullNetworkReplayStreamer::StartStreaming( const FString& StreamName, bool bRecord, const FNetworkReplayVersion& ReplayVersion, const FOnStreamReadyDelegate& Delegate )
 {
 	// Create a directory for this demo
 	const FString DemoDir = GetStreamDirectory(StreamName);
@@ -169,7 +169,7 @@ void FNullNetworkReplayStreamer::DeleteFinishedStream( const FString& StreamName
 	Delegate.ExecuteIfBound(DeleteSucceeded);
 }
 
-void FNullNetworkReplayStreamer::EnumerateStreams( const FString& VersionString, const FOnEnumerateStreamsComplete& Delegate )
+void FNullNetworkReplayStreamer::EnumerateStreams( const FNetworkReplayVersion& ReplayVersion, const FOnEnumerateStreamsComplete& Delegate )
 {
 	// Simply returns a stream for each folder in the Saved/Demos directory
 	const FString WildCardPath = GetDemoPath() + TEXT( "*" );
diff --git a/Engine/Source/Runtime/NetworkReplayStreaming/NullNetworkReplayStreaming/Public/NullNetworkReplayStreaming.h b/Engine/Source/Runtime/NetworkReplayStreaming/NullNetworkReplayStreaming/Public/NullNetworkReplayStreaming.h
--- a/Engine/Source/Runtime/NetworkReplayStreaming/NullNetworkReplayStreaming/Public/NullNetworkReplayStreaming.h
+++ b/Engine/Source/Runtime/NetworkReplayStreaming/NullNetworkReplayStreaming/Public/NullNetworkReplayStreaming.h
@@ -16,7 +16,7 @@ public:
 	{}
 	
 	/** INetworkReplayStreamer implementation */
-	virtual void StartStreaming( const FString& StreamName, bool bRecord, const FString& VersionString, const FOnStreamReadyDelegate& Delegate ) override;
+	virtual void StartStreaming( const FString& StreamName, bool bRecord, const FNetworkReplayVersion& ReplayVersion, const FOnStreamReadyDelegate& Delegate ) override;
 	virtual void StopStreaming() override;
 	virtual FArchive* GetHeaderArchive() override;
 	virtual FArchive* GetStreamingArchive() override;
@@ -32,7 +32,7 @@ public:
 	virtual bool IsDataAvailableForTimeRange( const uint32 StartTimeInMS, const uint32 EndTimeInMS ) override { return true; }
 	virtual bool IsLive( const FString& StreamName ) const override;
 	virtual void DeleteFinishedStream( const FString& StreamName, const FOnDeleteFinishedStreamComplete& Delegate) const override;
-	virtual void EnumerateStreams( const FString& VersionString, const FOnEnumerateStreamsComplete& Delegate ) override;
+	virtual void EnumerateStreams( const FNetworkReplayVersion& ReplayVersion, const FOnEnumerateStreamsComplete& Delegate ) override;
 	virtual ENetworkReplayError::Type GetLastError() const override { return ENetworkReplayError::None; }
 
 private:
