Pause replays as soon as we scrub

diff --git a/Engine/Source/Runtime/NetworkReplayStreaming/HttpNetworkReplayStreaming/Private/HttpNetworkReplayStreaming.cpp b/Engine/Source/Runtime/NetworkReplayStreaming/HttpNetworkReplayStreaming/Private/HttpNetworkReplayStreaming.cpp
--- a/Engine/Source/Runtime/NetworkReplayStreaming/HttpNetworkReplayStreaming/Private/HttpNetworkReplayStreaming.cpp
+++ b/Engine/Source/Runtime/NetworkReplayStreaming/HttpNetworkReplayStreaming/Private/HttpNetworkReplayStreaming.cpp
@@ -662,6 +662,11 @@ bool FHttpNetworkReplayStreamer::IsDataAvailable() const
 		return false;
 	}
 
+	if ( GotoCheckpointDelegate.IsBound() )
+	{
+		return false;
+	}
+
 	if ( HighPriorityEndTime > 0 )
 	{
 		// If we are waiting for a high priority portion of the stream, pretend like we don't have any data so that game code waits for the entire portion
