Implement checkpoints and fine-grained scrubbing in FNullNetworkReplayStreamer.
FNullNetworkReplayStreamer now stores the demo header in a separate file to facilitate scrubbing to the beginning of the replay stream.
DemoNetDriver now handles errors when StartStreaming calls the FOnStreamDelegate inside its implementation. Fixes getting stuck on the loading screen in ShooterGame if FNullNetworkReplayStreamer::StartStreaming fails.

diff --git a/Engine/Source/Runtime/Engine/Classes/Engine/DemoNetDriver.h b/Engine/Source/Runtime/Engine/Classes/Engine/DemoNetDriver.h
--- a/Engine/Source/Runtime/Engine/Classes/Engine/DemoNetDriver.h
+++ b/Engine/Source/Runtime/Engine/Classes/Engine/DemoNetDriver.h
@@ -80,6 +80,7 @@ class ENGINE_API UDemoNetDriver : public UNetDriver
 private:
 	bool		bIsFastForwarding;
 	bool		bIsLoadingCheckpoint;
+	bool		bWasStartStreamingSuccessful;
 
 	TArray<FNetworkGUID> NonQueuedGUIDsForScrubbing;
 
diff --git a/Engine/Source/Runtime/Engine/Private/DemoNetDriver.cpp b/Engine/Source/Runtime/Engine/Private/DemoNetDriver.cpp
--- a/Engine/Source/Runtime/Engine/Private/DemoNetDriver.cpp
+++ b/Engine/Source/Runtime/Engine/Private/DemoNetDriver.cpp
@@ -61,6 +61,7 @@ bool UDemoNetDriver::InitBase( bool bInitAsClient, FNetworkNotify* InNotify, con
 		QueuedGotoTimeInSeconds	= -1.0f;
 		bIsLoadingCheckpoint	= false;
 		InitialLiveDemoTime		= 0;
+		bWasStartStreamingSuccessful = true;
 
 		ResetDemoState();
 
@@ -191,9 +192,10 @@ bool UDemoNetDriver::InitConnect( FNetworkNotify* InNotify, const FURL& ConnectU
 	ServerConnection = NewObject<UNetConnection>(GetTransientPackage(), UDemoNetConnection::StaticClass());
 	ServerConnection->InitConnection( this, USOCK_Pending, ConnectURL, 1000000 );
 
+	bWasStartStreamingSuccessful = true;
 	ReplayStreamer->StartStreaming( DemoFilename, false, FNetworkVersion::GetReplayVersion(), FOnStreamReadyDelegate::CreateUObject( this, &UDemoNetDriver::ReplayStreamingReady ) );
 
-	return true;
+	return bWasStartStreamingSuccessful;
 }
 
 bool UDemoNetDriver::InitConnectInternal( FString& Error )
@@ -1312,6 +1314,8 @@ void UDemoNetDriver::SpawnDemoRecSpectator( UNetConnection* Connection )
 
 void UDemoNetDriver::ReplayStreamingReady( bool bSuccess, bool bRecord )
 {
+	bWasStartStreamingSuccessful = bSuccess;
+
 	if ( !bSuccess )
 	{
 		GetWorld()->GetGameInstance()->HandleDemoPlaybackFailure( EDemoPlayFailure::DemoNotFound, FString( EDemoPlayFailure::ToString( EDemoPlayFailure::DemoNotFound ) ) );
@@ -1425,7 +1429,7 @@ void UDemoNetDriver::LoadCheckpoint()
 	GuidCache->ObjectLookup.Empty();
 	GuidCache->NetGUIDLookup.Empty();
 
-	if ( GotoCheckpointArchive->TotalSize() == 0 )
+	if ( GotoCheckpointArchive->TotalSize() == 0 || GotoCheckpointArchive->TotalSize() == INDEX_NONE )
 	{
 		// This is the very first checkpoint, we'll read the stream from the very beginning in this case
 		DemoCurrentTime			= 0;
diff --git a/Engine/Source/Runtime/NetworkReplayStreaming/NullNetworkReplayStreaming/NullNetworkReplayStreaming.build.cs b/Engine/Source/Runtime/NetworkReplayStreaming/NullNetworkReplayStreaming/NullNetworkReplayStreaming.build.cs
--- a/Engine/Source/Runtime/NetworkReplayStreaming/NullNetworkReplayStreaming/NullNetworkReplayStreaming.build.cs
+++ b/Engine/Source/Runtime/NetworkReplayStreaming/NullNetworkReplayStreaming/NullNetworkReplayStreaming.build.cs
@@ -8,11 +8,14 @@ namespace UnrealBuildTool.Rules
 		{
 			PrivateIncludePaths.Add( "Runtime/NetworkReplayStreaming/NullNetworkReplayStreaming/Private" );
 
-			PublicDependencyModuleNames.AddRange(
+            PrivateIncludePathModuleNames.Add("OnlineSubsystem");
+
+			PrivateDependencyModuleNames.AddRange(
 				new string[]
 				{
 					"Core",
 					"NetworkReplayStreaming",
+                    "Json"
 				}
 			);
 		}
diff --git a/Engine/Source/Runtime/NetworkReplayStreaming/NullNetworkReplayStreaming/Private/NullNetworkReplayStreaming.cpp b/Engine/Source/Runtime/NetworkReplayStreaming/NullNetworkReplayStreaming/Private/NullNetworkReplayStreaming.cpp
--- a/Engine/Source/Runtime/NetworkReplayStreaming/NullNetworkReplayStreaming/Private/NullNetworkReplayStreaming.cpp
+++ b/Engine/Source/Runtime/NetworkReplayStreaming/NullNetworkReplayStreaming/Private/NullNetworkReplayStreaming.cpp
@@ -3,9 +3,31 @@
 #include "NullNetworkReplayStreaming.h"
 #include "Paths.h"
 #include "EngineVersion.h"
+#include "OnlineJsonSerializer.h"
 
 DEFINE_LOG_CATEGORY_STATIC( LogNullReplay, Log, All );
 
+/* Class to hold stream event information */
+class FNullCheckpointListItem : public FOnlineJsonSerializable
+{
+public:
+	FNullCheckpointListItem() : Time1(0), Time2(0) {}
+
+	FString		Group;
+	FString		Metadata;
+	uint32		Time1;
+	uint32		Time2;
+
+	// FOnlineJsonSerializable
+	BEGIN_ONLINE_JSON_SERIALIZER
+		ONLINE_JSON_SERIALIZE( "group",			Group );
+		ONLINE_JSON_SERIALIZE( "meta",			Metadata );
+		ONLINE_JSON_SERIALIZE( "time1",			Time1 );
+		ONLINE_JSON_SERIALIZE( "time2",			Time2 );
+	END_ONLINE_JSON_SERIALIZER
+};
+
+
 /**
  * Very basic implementation of network replay streaming using the file system
  * As of now, there is just simple opening and closing of the stream, and handing off the stream for direct use
@@ -55,6 +77,11 @@ static FString GetStreamFullBaseFilename(const FString& StreamName)
 	return FPaths::Combine(*GetStreamDirectory(StreamName), *GetStreamBaseFilename(StreamName));
 }
 
+static FString GetHeaderFilename(const FString& StreamName)
+{
+	return GetStreamFullBaseFilename(StreamName) + TEXT(".header");
+}
+
 static FString GetDemoFilename(const FString& StreamName)
 {
 	return GetStreamFullBaseFilename(StreamName) + TEXT(".demo");
@@ -65,6 +92,16 @@ static FString GetMetadataFilename(const FString& StreamName)
 	return GetStreamFullBaseFilename(StreamName) + TEXT(".metadata");
 }
 
+static FString GetCheckpointFilename( const FString& StreamName, int32 Index )
+{
+	return FPaths::Combine(*GetStreamDirectory(StreamName), TEXT("checkpoints"), *FString::Printf( TEXT("checkpoint%d"), Index ) );
+}
+
+static FString GetEventFilename( const FString& StreamName, int32 Index )
+{
+	return FPaths::Combine(*GetStreamDirectory(StreamName), TEXT("events"), *FString::Printf( TEXT("event%d"), Index ) );
+}
+
 void FNullNetworkReplayStreamer::StartStreaming( const FString& StreamName, bool bRecord, const FNetworkReplayVersion& ReplayVersion, const FOnStreamReadyDelegate& Delegate )
 {
 	// Create a directory for this demo
@@ -72,6 +109,7 @@ void FNullNetworkReplayStreamer::StartStreaming( const FString& StreamName, bool
 
 	IFileManager::Get().MakeDirectory( *DemoDir, true );
 
+	const FString FullHeaderFilename = GetHeaderFilename(StreamName);
 	const FString FullDemoFilename = GetDemoFilename(StreamName);
 	const FString FullMetadataFilename = GetMetadataFilename(StreamName);
 
@@ -81,21 +119,26 @@ void FNullNetworkReplayStreamer::StartStreaming( const FString& StreamName, bool
 	{
 		// Open file for reading
 		FileAr.Reset( IFileManager::Get().CreateFileReader( *FullDemoFilename ) );
+		HeaderAr.Reset( IFileManager::Get().CreateFileReader( *FullHeaderFilename ) );
 		StreamerState = EStreamerState::Playback;
 	}
 	else
 	{
 		// Open file for writing
 		FileAr.Reset( IFileManager::Get().CreateFileWriter( *FullDemoFilename, FILEWRITE_AllowRead ) );
+		HeaderAr.Reset( IFileManager::Get().CreateFileWriter( *FullHeaderFilename ) );
 		StreamerState = EStreamerState::Recording;
+
+		CurrentCheckpointIndex = 0;
 	}
 
 	// Notify immediately
-	Delegate.ExecuteIfBound( FileAr.Get() != NULL, bRecord );
+	Delegate.ExecuteIfBound( FileAr.Get() != nullptr && HeaderAr.Get() != nullptr, bRecord );
 }
 
 void FNullNetworkReplayStreamer::StopStreaming()
 {
+	HeaderAr.Reset();
 	FileAr.Reset();
 	MetadataFileAr.Reset();
 
@@ -105,7 +148,7 @@ void FNullNetworkReplayStreamer::StopStreaming()
 
 FArchive* FNullNetworkReplayStreamer::GetHeaderArchive()
 {
-	return FileAr.Get();
+	return HeaderAr.Get();
 }
 
 FArchive* FNullNetworkReplayStreamer::GetStreamingArchive()
@@ -209,6 +252,169 @@ void FNullNetworkReplayStreamer::EnumerateStreams( const FNetworkReplayVersion&
 	Delegate.ExecuteIfBound(Results);
 }
 
+FArchive* FNullNetworkReplayStreamer::GetCheckpointArchive()
+{
+	// If the archive is null, and the API is being used properly, the caller is writing a checkpoint...
+	if ( CheckpointAr.Get() == nullptr )
+	{
+		// Create a file writer for the next checkpoint index.
+		check(StreamerState != EStreamerState::Playback);
+
+		FString NextCheckpointFileName = FString::Printf( TEXT( "checkpoint%d" ), CurrentCheckpointIndex );
+
+		UE_LOG(LogNullReplay, Log, TEXT("FNullNetworkReplayStreamer::GetCheckpointArchive. Creating new checkpoint file."));
+
+		CheckpointAr.Reset( IFileManager::Get().CreateFileWriter( *GetCheckpointFilename(CurrentStreamName, CurrentCheckpointIndex) ) );
+	}
+
+	return CheckpointAr.Get();
+}
+
+void FNullNetworkReplayStreamer::FlushCheckpoint(const uint32 TimeInMS)
+{
+	UE_LOG(LogNullReplay, Log, TEXT("FNullNetworkReplayStreamer::FlushCheckpoint. TimeInMS: %u"), TimeInMS);
+
+	check(FileAr.Get() != nullptr);
+
+	// The file writer archive will finalize the file on disk on destruction. The new file will be created
+	// next time the driver calls GetCheckpointArchive.
+	CheckpointAr.Reset();
+
+	// Also write the event description file to disk with a corresponding checkpoint index, so they can be correlated later.
+	TUniquePtr<FArchive> EventFileAr(IFileManager::Get().CreateFileWriter(*GetEventFilename(CurrentStreamName, CurrentCheckpointIndex)));
+
+
+	if (EventFileAr.Get() != nullptr)
+	{
+		FNullCheckpointListItem CheckpointEvent;
+		CheckpointEvent.Group = TEXT("checkpoint");
+		CheckpointEvent.Metadata = FString::Printf( TEXT("%ld"), FileAr->Tell() );
+		CheckpointEvent.Time1 = TimeInMS;
+		CheckpointEvent.Time2 = TimeInMS;
+
+		FString EventJsonString = CheckpointEvent.ToJson();
+		*EventFileAr << EventJsonString;
+	}
+
+	++CurrentCheckpointIndex;
+}
+
+void FNullNetworkReplayStreamer::GotoCheckpointIndex(const int32 CheckpointIndex, const FOnCheckpointReadyDelegate& Delegate)
+{
+	GotoCheckpointIndexInternal(CheckpointIndex, Delegate, -1);
+}
+
+void FNullNetworkReplayStreamer::GotoCheckpointIndexInternal(int32 CheckpointIndex, const FOnCheckpointReadyDelegate& Delegate, int32 TimeInMS)
+{
+	check( FileAr.Get() != nullptr);
+
+	if ( CheckpointIndex == -1 )
+	{
+		// Create a dummy checkpoint archive to indicate this is the first checkpoint
+		CheckpointAr.Reset(new FArchive);
+
+		FileAr->Seek(0);
+		
+		Delegate.ExecuteIfBound( true, TimeInMS );
+		return;
+	}
+
+	// Attempt to open the checkpoint file for the given index. Will fail if file doesn't exist.
+	const FString CheckpointFilename = GetCheckpointFilename(CurrentStreamName, CheckpointIndex);
+	CheckpointAr.Reset( IFileManager::Get().CreateFileReader( *CheckpointFilename ) );
+
+	if ( CheckpointAr.Get() == nullptr )
+	{
+		UE_LOG(LogNullReplay, Log, TEXT("FNullNetworkReplayStreamer::GotoCheckpointIndex. Index: %i. Couldn't open checkpoint file %s"), CheckpointIndex, *CheckpointFilename);
+		Delegate.ExecuteIfBound( false, TimeInMS );
+		return;
+	}
+
+	// Open and deserialize the corresponding event, this tells us where we need to seek to
+	// in the main replay file to sync up with the checkpoint we're loading.
+	const FString EventFilename = GetEventFilename(CurrentStreamName, CheckpointIndex);
+	TUniquePtr<FArchive> EventFile( IFileManager::Get().CreateFileReader(*EventFilename));
+	if (EventFile.Get() != nullptr)
+	{
+		FString JsonString;
+		*EventFile << JsonString;
+
+		FNullCheckpointListItem Item;
+		Item.FromJson(JsonString);
+
+		FileAr->Seek( FCString::Atoi64( *Item.Metadata ) );
+	}
+
+	Delegate.ExecuteIfBound( true, TimeInMS );
+}
+
+void FNullNetworkReplayStreamer::GotoTimeInMS(const uint32 TimeInMS, const FOnCheckpointReadyDelegate& Delegate)
+{
+	// Enumerate all the events in the events folder, since we need to know what times the checkpoints correlate with
+	TArray<FNullCheckpointListItem> Checkpoints;
+
+	const FString EventBaseName = FPaths::Combine( *GetStreamDirectory(CurrentStreamName), TEXT( "events" ), TEXT( "event" ) );
+
+	int CurrentEventIndex = 0;
+
+	// Try to load every event in order until one is missing
+	while ( true )
+	{
+		const FString CheckEventName = EventBaseName + FString::FromInt(CurrentEventIndex);
+
+		TUniquePtr<FArchive> EventFile( IFileManager::Get().CreateFileReader(*CheckEventName) );
+
+		if ( EventFile.Get() != nullptr )
+		{
+			FString JsonString;
+			*EventFile << JsonString;
+
+			FNullCheckpointListItem Item;
+			Item.FromJson(JsonString);
+
+			Checkpoints.Add(Item);
+		}
+		else
+		{
+			break;
+		}
+
+		CurrentEventIndex++;
+	}
+
+	int32 CheckpointIndex = -1;
+
+	if ( Checkpoints.Num() > 0 && TimeInMS >= Checkpoints[ Checkpoints.Num() - 1 ].Time1 )
+	{
+		// If we're after the very last checkpoint, that's the one we want
+		CheckpointIndex = Checkpoints.Num() - 1;
+	}
+	else
+	{
+		// Checkpoints should be sorted by time, return the checkpoint that exists right before the current time
+		// For fine scrubbing, we'll fast forward the rest of the way
+		// NOTE - If we're right before the very first checkpoint, we'll return -1, which is what we want when we want to start from the very beginning
+		for ( int i = 0; i < Checkpoints.Num(); i++ )
+		{
+			if ( TimeInMS < Checkpoints[i].Time1 )
+			{
+				CheckpointIndex = i - 1;
+				break;
+			}
+		}
+	}
+
+	int32 ExtraSkipTimeInMS = TimeInMS;
+
+	if ( CheckpointIndex >= 0 )
+	{
+		// Subtract off checkpoint time so we pass in the leftover to the engine to fast forward through for the fine scrubbing part
+		ExtraSkipTimeInMS = TimeInMS - Checkpoints[ CheckpointIndex ].Time1;
+	}
+
+	GotoCheckpointIndexInternal( CheckpointIndex, Delegate, ExtraSkipTimeInMS );
+}
+
 IMPLEMENT_MODULE( FNullNetworkReplayStreamingFactory, NullNetworkReplayStreaming )
 
 TSharedPtr< INetworkReplayStreamer > FNullNetworkReplayStreamingFactory::CreateReplayStreamer() 
diff --git a/Engine/Source/Runtime/NetworkReplayStreaming/NullNetworkReplayStreaming/Public/NullNetworkReplayStreaming.h b/Engine/Source/Runtime/NetworkReplayStreaming/NullNetworkReplayStreaming/Public/NullNetworkReplayStreaming.h
--- a/Engine/Source/Runtime/NetworkReplayStreaming/NullNetworkReplayStreaming/Public/NullNetworkReplayStreaming.h
+++ b/Engine/Source/Runtime/NetworkReplayStreaming/NullNetworkReplayStreaming/Public/NullNetworkReplayStreaming.h
@@ -12,7 +12,8 @@ class FNullNetworkReplayStreamer : public INetworkReplayStreamer
 {
 public:
 	FNullNetworkReplayStreamer() :
-		StreamerState( EStreamerState::Idle )
+		StreamerState( EStreamerState::Idle ),
+		CurrentCheckpointIndex( 0 )
 	{}
 	
 	/** INetworkReplayStreamer implementation */
@@ -20,10 +21,10 @@ public:
 	virtual void StopStreaming() override;
 	virtual FArchive* GetHeaderArchive() override;
 	virtual FArchive* GetStreamingArchive() override;
-	virtual FArchive* GetCheckpointArchive() override { return NULL; }
-	virtual void FlushCheckpoint( const uint32 TimeInMS ) override { }
-	virtual void GotoCheckpointIndex( const int32 TimeInMS, const FOnCheckpointReadyDelegate& Delegate ) override { }
-	virtual void GotoTimeInMS( const uint32 TimeInMS, const FOnCheckpointReadyDelegate& Delegate ) override { }
+	virtual FArchive* GetCheckpointArchive() override;
+	virtual void FlushCheckpoint( const uint32 TimeInMS ) override;
+	virtual void GotoCheckpointIndex( const int32 CheckpointIndex, const FOnCheckpointReadyDelegate& Delegate ) override;
+	virtual void GotoTimeInMS( const uint32 TimeInMS, const FOnCheckpointReadyDelegate& Delegate ) override;
 	virtual FArchive* GetMetadataArchive() override;
 	virtual void UpdateTotalDemoTime( uint32 TimeInMS ) override { }
 	virtual uint32 GetTotalDemoTime() const override { return 0; }
@@ -39,12 +40,21 @@ public:
 private:
 	bool IsNamedStreamLive( const FString& StreamName ) const;
 
+	/** Handles the details of loading a checkpoint */
+	void GotoCheckpointIndexInternal(int32 CheckpointIndex, const FOnCheckpointReadyDelegate& Delegate, int32 TimeInMS);
+
+	/** Handle to the archive that will read/write the demo header */
+	TUniquePtr<FArchive> HeaderAr;
+
 	/** Handle to the archive that will read/write network packets */
 	TUniquePtr<FArchive> FileAr;
 
 	/* Handle to the archive that will read/write metadata */
 	TUniquePtr<FArchive> MetadataFileAr;
 
+	/* Handle to the archive that will read/write checkpoint files */
+	TUniquePtr<FArchive> CheckpointAr;
+
 	/** EStreamerState - Overall state of the streamer */
 	enum class EStreamerState
 	{
@@ -58,6 +68,9 @@ private:
 
 	/** Remember the name of the current stream, if any. */
 	FString CurrentStreamName;
+
+	/** Current number of checkpoints written. */
+	int32 CurrentCheckpointIndex;
 };
 
 class FNullNetworkReplayStreamingFactory : public INetworkReplayStreamingFactory
