INetworkReplayStreamer::StartStreaming now takes a CustomName and a FriendlyName
Demorec command and URL option name parameter is now optional, if it isn't provided the streamer implementation will automatically generate CustomName, and the FriendlyName will be the map name
Logic to auto-generate demo names moved out of ShooterGame and into NullNetworkReplayStreaming

diff --git a/Engine/Source/Runtime/Engine/Private/DemoNetDriver.cpp b/Engine/Source/Runtime/Engine/Private/DemoNetDriver.cpp
--- a/Engine/Source/Runtime/Engine/Private/DemoNetDriver.cpp
+++ b/Engine/Source/Runtime/Engine/Private/DemoNetDriver.cpp
@@ -193,7 +193,8 @@ bool UDemoNetDriver::InitConnect( FNetworkNotify* InNotify, const FURL& ConnectU
 	ServerConnection->InitConnection( this, USOCK_Pending, ConnectURL, 1000000 );
 
 	bWasStartStreamingSuccessful = true;
-	ReplayStreamer->StartStreaming( DemoFilename, false, FNetworkVersion::GetReplayVersion(), FOnStreamReadyDelegate::CreateUObject( this, &UDemoNetDriver::ReplayStreamingReady ) );
+	// Friendly name isn't important for loading an existing replay.
+	ReplayStreamer->StartStreaming( DemoFilename, FString(), false, FNetworkVersion::GetReplayVersion(), FOnStreamReadyDelegate::CreateUObject( this, &UDemoNetDriver::ReplayStreamingReady ) );
 
 	return bWasStartStreamingSuccessful;
 }
@@ -365,7 +366,7 @@ bool UDemoNetDriver::InitListen( FNetworkNotify* InNotify, FURL& ListenURL, bool
 	Connection->InitSendBuffer();
 	ClientConnections.Add( Connection );
 
-	ReplayStreamer->StartStreaming( DemoFilename, true, FNetworkVersion::GetReplayVersion(), FOnStreamReadyDelegate::CreateUObject( this, &UDemoNetDriver::ReplayStreamingReady ) );
+	ReplayStreamer->StartStreaming( DemoFilename, World->GetMapName(), true, FNetworkVersion::GetReplayVersion(), FOnStreamReadyDelegate::CreateUObject( this, &UDemoNetDriver::ReplayStreamingReady ) );
 
 	FArchive* FileAr = ReplayStreamer->GetHeaderArchive();
 
diff --git a/Engine/Source/Runtime/Engine/Private/UnrealEngine.cpp b/Engine/Source/Runtime/Engine/Private/UnrealEngine.cpp
--- a/Engine/Source/Runtime/Engine/Private/UnrealEngine.cpp
+++ b/Engine/Source/Runtime/Engine/Private/UnrealEngine.cpp
@@ -9308,12 +9308,12 @@ bool UEngine::LoadMap( FWorldContext& WorldContext, FURL URL, class UPendingNetG
 	// See if we need to record network demos
 	if ( WorldContext.World()->GetAuthGameMode() == NULL || !WorldContext.World()->GetAuthGameMode()->IsHandlingReplays() )
 	{
-		const TCHAR* DemoRecName = URL.GetOption( TEXT( "DemoRec=" ), NULL );
-
-		if ( DemoRecName != NULL )
+		if ( URL.HasOption( TEXT( "DemoRec" ) ) )
 		{
-			// Record the demo
-			GEngine->Exec( WorldContext.World(), *FString::Printf( TEXT("DEMOREC %s"), DemoRecName ) );
+			const TCHAR* DemoRecName = URL.GetOption( TEXT( "DemoRec=" ), NULL );
+
+			// Record the demo, optionally with the specified custom name.
+			GEngine->Exec( WorldContext.World(), *FString::Printf( TEXT("DEMOREC %s"), DemoRecName != NULL ? DemoRecName : TEXT("") ) );
 		}
 	}
 
diff --git a/Engine/Source/Runtime/Engine/Private/World.cpp b/Engine/Source/Runtime/Engine/Private/World.cpp
--- a/Engine/Source/Runtime/Engine/Private/World.cpp
+++ b/Engine/Source/Runtime/Engine/Private/World.cpp
@@ -2644,12 +2644,8 @@ bool UWorld::HandleDemoRecordCommand( const TCHAR* Cmd, FOutputDevice& Ar, UWorl
 	FURL DemoURL;
 	FString DemoName;
 
-	if ( !FParse::Token( Cmd, DemoName, 0 ) )
-	{
-		Ar.Logf( TEXT( "Missing demo name." ) );
-		return true;
-	}
-
+	FParse::Token( Cmd, DemoName, 0 );
+	
 	DemoName.ReplaceInline( TEXT( "%m" ), *GetMapName() );
 
 	// replace the current URL's map with a demo extension
diff --git a/Engine/Source/Runtime/NetworkReplayStreaming/HttpNetworkReplayStreaming/Private/HttpNetworkReplayStreaming.cpp b/Engine/Source/Runtime/NetworkReplayStreaming/HttpNetworkReplayStreaming/Private/HttpNetworkReplayStreaming.cpp
--- a/Engine/Source/Runtime/NetworkReplayStreaming/HttpNetworkReplayStreaming/Private/HttpNetworkReplayStreaming.cpp
+++ b/Engine/Source/Runtime/NetworkReplayStreaming/HttpNetworkReplayStreaming/Private/HttpNetworkReplayStreaming.cpp
@@ -122,7 +122,7 @@ FHttpNetworkReplayStreamer::FHttpNetworkReplayStreamer() :
 	GConfig->GetString( TEXT( "HttpNetworkReplayStreaming" ), TEXT( "ServerURL" ), ServerURL, GEngineIni );
 }
 
-void FHttpNetworkReplayStreamer::StartStreaming( const FString& StreamName, bool bRecord, const FNetworkReplayVersion& InReplayVersion, const FOnStreamReadyDelegate& Delegate )
+void FHttpNetworkReplayStreamer::StartStreaming( const FString& CustomName, const FString& FriendlyName, bool bRecord, const FNetworkReplayVersion& InReplayVersion, const FOnStreamReadyDelegate& Delegate )
 {
 	if ( !SessionName.IsEmpty() )
 	{
@@ -175,7 +175,7 @@ void FHttpNetworkReplayStreamer::StartStreaming( const FString& StreamName, bool
 		// We are streaming down
 		StreamerState = EStreamerState::StreamingDown;
 
-		SessionName = StreamName;
+		SessionName = CustomName;
 
 		// Notify the http server that we want to start downloading a replay
 		HttpRequest->SetURL( FString::Printf( TEXT( "%sstartdownloading?Session=%s" ), *ServerURL, *SessionName ) );
@@ -203,7 +203,7 @@ void FHttpNetworkReplayStreamer::StartStreaming( const FString& StreamName, bool
 		bNeedToUploadHeader = true;
 
 		// Notify the http server that we want to start uploading a replay
-		HttpRequest->SetURL( FString::Printf( TEXT( "%sstartuploading?App=%s&Version=%u&CL=%u&Friendly=%s" ), *ServerURL, *ReplayVersion.AppString, ReplayVersion.NetworkVersion, ReplayVersion.Changelist, *StreamName ) );
+		HttpRequest->SetURL( FString::Printf( TEXT( "%sstartuploading?App=%s&Version=%u&CL=%u&Friendly=%s" ), *ServerURL, *ReplayVersion.AppString, ReplayVersion.NetworkVersion, ReplayVersion.Changelist, *FriendlyName ) );
 
 		HttpRequest->SetVerb( TEXT( "POST" ) );
 
diff --git a/Engine/Source/Runtime/NetworkReplayStreaming/HttpNetworkReplayStreaming/Private/HttpNetworkReplayStreaming.h b/Engine/Source/Runtime/NetworkReplayStreaming/HttpNetworkReplayStreaming/Private/HttpNetworkReplayStreaming.h
--- a/Engine/Source/Runtime/NetworkReplayStreaming/HttpNetworkReplayStreaming/Private/HttpNetworkReplayStreaming.h
+++ b/Engine/Source/Runtime/NetworkReplayStreaming/HttpNetworkReplayStreaming/Private/HttpNetworkReplayStreaming.h
@@ -96,7 +96,7 @@ public:
 	FHttpNetworkReplayStreamer();
 
 	/** INetworkReplayStreamer implementation */
-	virtual void		StartStreaming( const FString& StreamName, bool bRecord, const FNetworkReplayVersion& ReplayVersion, const FOnStreamReadyDelegate& Delegate ) override;
+	virtual void		StartStreaming( const FString& CustomName, const FString& FriendlyName, bool bRecord, const FNetworkReplayVersion& ReplayVersion, const FOnStreamReadyDelegate& Delegate ) override;
 	virtual void		StopStreaming() override;
 	virtual FArchive*	GetHeaderArchive() override;
 	virtual FArchive*	GetStreamingArchive() override;
diff --git a/Engine/Source/Runtime/NetworkReplayStreaming/NetworkReplayStreaming/Public/NetworkReplayStreaming.h b/Engine/Source/Runtime/NetworkReplayStreaming/NetworkReplayStreaming/Public/NetworkReplayStreaming.h
--- a/Engine/Source/Runtime/NetworkReplayStreaming/NetworkReplayStreaming/Public/NetworkReplayStreaming.h
+++ b/Engine/Source/Runtime/NetworkReplayStreaming/NetworkReplayStreaming/Public/NetworkReplayStreaming.h
@@ -103,7 +103,7 @@ class INetworkReplayStreamer
 public:
 	virtual ~INetworkReplayStreamer() {}
 
-	virtual void StartStreaming( const FString& StreamName, bool bRecord, const FNetworkReplayVersion& ReplayVersion, const FOnStreamReadyDelegate& Delegate ) = 0;
+	virtual void StartStreaming( const FString& CustomName, const FString& FriendlyName, bool bRecord, const FNetworkReplayVersion& ReplayVersion, const FOnStreamReadyDelegate& Delegate ) = 0;
 	virtual void StopStreaming() = 0;
 	virtual FArchive* GetHeaderArchive() = 0;
 	virtual FArchive* GetStreamingArchive() = 0;
diff --git a/Engine/Source/Runtime/NetworkReplayStreaming/NullNetworkReplayStreaming/Private/NullNetworkReplayStreaming.cpp b/Engine/Source/Runtime/NetworkReplayStreaming/NullNetworkReplayStreaming/Private/NullNetworkReplayStreaming.cpp
--- a/Engine/Source/Runtime/NetworkReplayStreaming/NullNetworkReplayStreaming/Private/NullNetworkReplayStreaming.cpp
+++ b/Engine/Source/Runtime/NetworkReplayStreaming/NullNetworkReplayStreaming/Private/NullNetworkReplayStreaming.cpp
@@ -102,18 +102,66 @@ static FString GetEventFilename( const FString& StreamName, int32 Index )
 	return FPaths::Combine(*GetStreamDirectory(StreamName), TEXT("events"), *FString::Printf( TEXT("event%d"), Index ) );
 }
 
-void FNullNetworkReplayStreamer::StartStreaming( const FString& StreamName, bool bRecord, const FNetworkReplayVersion& ReplayVersion, const FOnStreamReadyDelegate& Delegate )
+// Returns a name formatted as "demoX", where X is 0-9.
+// Returns the first value that doesn't yet exist, or if they all exist, returns the oldest one
+// (it will be overwritten).
+static FString GetAutomaticDemoName()
 {
-	// Create a directory for this demo
-	const FString DemoDir = GetStreamDirectory(StreamName);
+	FString FinalDemoName;
+	FDateTime BestDateTime = FDateTime::MaxValue();
 
-	IFileManager::Get().MakeDirectory( *DemoDir, true );
+	const int MAX_DEMOS = 10;
 
-	const FString FullHeaderFilename = GetHeaderFilename(StreamName);
-	const FString FullDemoFilename = GetDemoFilename(StreamName);
-	const FString FullMetadataFilename = GetMetadataFilename(StreamName);
+	for (int32 i = 0; i < MAX_DEMOS; i++)
+	{
+		const FString DemoName = FString::Printf(TEXT("demo%i"), i + 1);
+		
+		const FString FullDemoName = GetDemoFilename(DemoName);
+		
+		FDateTime DateTime = IFileManager::Get().GetTimeStamp(*FullDemoName);
+
+		if (DateTime == FDateTime::MinValue())
+		{
+			// If we don't find this file, we can early out now
+			FinalDemoName = DemoName;
+			break;
+		}
+		else if (DateTime < BestDateTime)
+		{
+			// Use the oldest file
+			FinalDemoName = DemoName;
+			BestDateTime = DateTime;
+		}
+	}
+
+	return FinalDemoName;
+}
+
+void FNullNetworkReplayStreamer::StartStreaming( const FString& CustomName, const FString& FriendlyName, bool bRecord, const FNetworkReplayVersion& ReplayVersion, const FOnStreamReadyDelegate& Delegate )
+{
+	FString FinalDemoName = CustomName;
 
-	CurrentStreamName = StreamName;
+	if ( CustomName.IsEmpty() )
+	{
+		if ( bRecord )
+		{
+			// If we're recording and the caller didn't provide a name, generate one automatically
+			FinalDemoName = GetAutomaticDemoName();
+		}
+		else
+		{
+			// Can't play a replay if the user didn't provide a name!
+			Delegate.ExecuteIfBound( false, bRecord );
+			return;
+		}
+	}
+	
+	const FString DemoDir = GetStreamDirectory(FinalDemoName);
+
+	const FString FullHeaderFilename = GetHeaderFilename(FinalDemoName);
+	const FString FullDemoFilename = GetDemoFilename(FinalDemoName);
+	
+	CurrentStreamName = FinalDemoName;
 
 	if ( !bRecord )
 	{
@@ -124,6 +172,12 @@ void FNullNetworkReplayStreamer::StartStreaming( const FString& StreamName, bool
 	}
 	else
 	{
+		// Delete any existing demo with this name
+		IFileManager::Get().DeleteDirectory( *DemoDir, false, true );
+		
+		// Create a directory for this demo
+		IFileManager::Get().MakeDirectory( *DemoDir, true );
+
 		// Open file for writing
 		FileAr.Reset( IFileManager::Get().CreateFileWriter( *FullDemoFilename, FILEWRITE_AllowRead ) );
 		HeaderAr.Reset( IFileManager::Get().CreateFileWriter( *FullHeaderFilename ) );
diff --git a/Engine/Source/Runtime/NetworkReplayStreaming/NullNetworkReplayStreaming/Public/NullNetworkReplayStreaming.h b/Engine/Source/Runtime/NetworkReplayStreaming/NullNetworkReplayStreaming/Public/NullNetworkReplayStreaming.h
--- a/Engine/Source/Runtime/NetworkReplayStreaming/NullNetworkReplayStreaming/Public/NullNetworkReplayStreaming.h
+++ b/Engine/Source/Runtime/NetworkReplayStreaming/NullNetworkReplayStreaming/Public/NullNetworkReplayStreaming.h
@@ -17,7 +17,7 @@ public:
 	{}
 	
 	/** INetworkReplayStreamer implementation */
-	virtual void StartStreaming( const FString& StreamName, bool bRecord, const FNetworkReplayVersion& ReplayVersion, const FOnStreamReadyDelegate& Delegate ) override;
+	virtual void StartStreaming( const FString& CustomName, const FString& FriendlyName, bool bRecord, const FNetworkReplayVersion& ReplayVersion, const FOnStreamReadyDelegate& Delegate ) override;
 	virtual void StopStreaming() override;
 	virtual FArchive* GetHeaderArchive() override;
 	virtual FArchive* GetStreamingArchive() override;
