FNullNetworkReplayStreamer now stores some extra metadata about each replay, including the total time, network version, and changelist with which it was saved.
FNullNetworkReplayStreamer::EnumerateStreams now respects the version passed in (you can pass in 0 to enumerate all versions).
FNullNetworkReplayStreamer::GetTotalDemoTime now returns the correct value.
Fixes UE-13611 and UE-13612.

diff --git a/Engine/Source/Runtime/NetworkReplayStreaming/NullNetworkReplayStreaming/Private/NullNetworkReplayStreaming.cpp b/Engine/Source/Runtime/NetworkReplayStreaming/NullNetworkReplayStreaming/Private/NullNetworkReplayStreaming.cpp
--- a/Engine/Source/Runtime/NetworkReplayStreaming/NullNetworkReplayStreaming/Private/NullNetworkReplayStreaming.cpp
+++ b/Engine/Source/Runtime/NetworkReplayStreaming/NullNetworkReplayStreaming/Private/NullNetworkReplayStreaming.cpp
@@ -3,7 +3,6 @@
 #include "NullNetworkReplayStreaming.h"
 #include "Paths.h"
 #include "EngineVersion.h"
-#include "OnlineJsonSerializer.h"
 
 DEFINE_LOG_CATEGORY_STATIC( LogNullReplay, Log, All );
 
@@ -27,7 +26,6 @@ public:
 	END_ONLINE_JSON_SERIALIZER
 };
 
-
 /**
  * Very basic implementation of network replay streaming using the file system
  * As of now, there is just simple opening and closing of the stream, and handing off the stream for direct use
@@ -102,6 +100,29 @@ static FString GetEventFilename( const FString& StreamName, int32 Index )
 	return FPaths::Combine(*GetStreamDirectory(StreamName), TEXT("events"), *FString::Printf( TEXT("event%d"), Index ) );
 }
 
+static FString GetInfoFilename( const FString& StreamName )
+{
+	return GetStreamFullBaseFilename(StreamName) + TEXT(".replayinfo");
+}
+
+static FNullReplayInfo ReadReplayInfo( const FString& StreamName )
+{
+	FNullReplayInfo Info;
+
+	const FString InfoFilename = GetInfoFilename(StreamName);
+	TUniquePtr<FArchive> InfoFileArchive( IFileManager::Get().CreateFileReader( *InfoFilename ) );
+
+	if ( InfoFileArchive.IsValid() )
+	{
+		FString JsonString;
+		*InfoFileArchive << JsonString;
+
+		Info.FromJson(JsonString);
+	}
+
+	return Info;
+}
+
 // Returns a name formatted as "demoX", where X is 0-9.
 // Returns the first value that doesn't yet exist, or if they all exist, returns the oldest one
 // (it will be overwritten).
@@ -165,6 +186,9 @@ void FNullNetworkReplayStreamer::StartStreaming( const FString& CustomName, cons
 
 	if ( !bRecord )
 	{
+		// Load metadata if it exists
+		ReplayInfo = ReadReplayInfo( CurrentStreamName );
+
 		// Open file for reading
 		FileAr.Reset( IFileManager::Get().CreateFileReader( *FullDemoFilename ) );
 		HeaderAr.Reset( IFileManager::Get().CreateFileReader( *FullHeaderFilename ) );
@@ -184,6 +208,10 @@ void FNullNetworkReplayStreamer::StartStreaming( const FString& CustomName, cons
 		StreamerState = EStreamerState::Recording;
 
 		CurrentCheckpointIndex = 0;
+
+		// Set up replay info
+		ReplayInfo.NetworkVersion = ReplayVersion.NetworkVersion;
+		ReplayInfo.Changelist = ReplayVersion.Changelist;
 	}
 
 	// Notify immediately
@@ -192,6 +220,18 @@ void FNullNetworkReplayStreamer::StartStreaming( const FString& CustomName, cons
 
 void FNullNetworkReplayStreamer::StopStreaming()
 {
+	if (StreamerState == EStreamerState::Recording)
+	{
+		// Update metadata file with latest info
+		TUniquePtr<FArchive> ReplayInfoFileAr(IFileManager::Get().CreateFileWriter(*GetInfoFilename(CurrentStreamName)));
+
+		if (ReplayInfoFileAr.IsValid())
+		{
+			FString JsonString = ReplayInfo.ToJson();
+			*ReplayInfoFileAr << JsonString; 
+		}
+	}
+
 	HeaderAr.Reset();
 	FileAr.Reset();
 	MetadataFileAr.Reset();
@@ -235,6 +275,11 @@ FArchive* FNullNetworkReplayStreamer::GetMetadataArchive()
 	return MetadataFileAr.Get();
 }
 
+void FNullNetworkReplayStreamer::UpdateTotalDemoTime(uint32 TimeInMS)
+{
+	ReplayInfo.LengthInMS = TimeInMS;
+}
+
 bool FNullNetworkReplayStreamer::IsLive() const
 {
 	return IsNamedStreamLive(CurrentStreamName);
@@ -291,14 +336,28 @@ void FNullNetworkReplayStreamer::EnumerateStreams( const FNetworkReplayVersion&
 		
 		if (Info.SizeInBytes != INDEX_NONE)
 		{
-			Info.Name = Directory;
-			Info.Timestamp = IFileManager::Get().GetTimeStamp( *FullDemoFilePath );
-			Info.bIsLive = IsNamedStreamLive( Directory );			
+			// Read stored info for this replay
+			FNullReplayInfo StoredReplayInfo = ReadReplayInfo( Directory );
+			
+			// Check version. NetworkVersion and changelist of 0 will ignore version check.
+			const bool NetworkVersionMatches = ReplayVersion.NetworkVersion == StoredReplayInfo.NetworkVersion;
+			const bool ChangelistMatches = ReplayVersion.Changelist == StoredReplayInfo.Changelist;
+
+			const bool NetworkVersionPasses = ReplayVersion.NetworkVersion == 0 || NetworkVersionMatches;
+			const bool ChangelistPasses = ReplayVersion.Changelist == 0 || ChangelistMatches;
 
-			// Live streams not supported yet
-			if (!Info.bIsLive)
+			if ( NetworkVersionPasses && ChangelistPasses )
 			{
-				Results.Add(Info);
+				Info.Name = Directory;
+				Info.Timestamp = IFileManager::Get().GetTimeStamp( *FullDemoFilePath );
+				Info.bIsLive = IsNamedStreamLive( Directory );
+				Info.LengthInMS = StoredReplayInfo.LengthInMS;
+
+				// Live streams not supported yet
+				if (!Info.bIsLive)
+				{
+					Results.Add(Info);
+				}
 			}
 		}
 	}
diff --git a/Engine/Source/Runtime/NetworkReplayStreaming/NullNetworkReplayStreaming/Public/NullNetworkReplayStreaming.h b/Engine/Source/Runtime/NetworkReplayStreaming/NullNetworkReplayStreaming/Public/NullNetworkReplayStreaming.h
--- a/Engine/Source/Runtime/NetworkReplayStreaming/NullNetworkReplayStreaming/Public/NullNetworkReplayStreaming.h
+++ b/Engine/Source/Runtime/NetworkReplayStreaming/NullNetworkReplayStreaming/Public/NullNetworkReplayStreaming.h
@@ -6,6 +6,25 @@
 #include "Core.h"
 #include "ModuleManager.h"
 #include "UniquePtr.h"
+#include "OnlineJsonSerializer.h"
+
+/* Class to hold metadata about an entire replay */
+class FNullReplayInfo : public FOnlineJsonSerializable
+{
+public:
+	FNullReplayInfo() : LengthInMS(0), NetworkVersion(0), Changelist(0) {}
+
+	int32		LengthInMS;
+	uint32		NetworkVersion;
+	uint32		Changelist;
+	
+	// FOnlineJsonSerializable
+	BEGIN_ONLINE_JSON_SERIALIZER
+		ONLINE_JSON_SERIALIZE( "LengthInMS",		LengthInMS );
+		ONLINE_JSON_SERIALIZE( "NetworkVersion",	NetworkVersion );
+		ONLINE_JSON_SERIALIZE( "Changelist",		Changelist );
+	END_ONLINE_JSON_SERIALIZER
+};
 
 /** Default streamer that goes straight to the HD */
 class FNullNetworkReplayStreamer : public INetworkReplayStreamer
@@ -26,8 +45,8 @@ public:
 	virtual void GotoCheckpointIndex( const int32 CheckpointIndex, const FOnCheckpointReadyDelegate& Delegate ) override;
 	virtual void GotoTimeInMS( const uint32 TimeInMS, const FOnCheckpointReadyDelegate& Delegate ) override;
 	virtual FArchive* GetMetadataArchive() override;
-	virtual void UpdateTotalDemoTime( uint32 TimeInMS ) override { }
-	virtual uint32 GetTotalDemoTime() const override { return 0; }
+	virtual void UpdateTotalDemoTime( uint32 TimeInMS ) override;
+	virtual uint32 GetTotalDemoTime() const override { return ReplayInfo.LengthInMS; }
 	virtual bool IsDataAvailable() const override { return true; }
 	virtual void SetHighPriorityTimeRange( const uint32 StartTimeInMS, const uint32 EndTimeInMS ) override { }
 	virtual bool IsDataAvailableForTimeRange( const uint32 StartTimeInMS, const uint32 EndTimeInMS ) override { return true; }
@@ -71,6 +90,9 @@ private:
 
 	/** Current number of checkpoints written. */
 	int32 CurrentCheckpointIndex;
+
+	/** Currently playing or recording replay metadata */
+	FNullReplayInfo ReplayInfo;
 };
 
 class FNullNetworkReplayStreamingFactory : public INetworkReplayStreamingFactory
