Merge to main: Return changelist when enumerating replays

diff --git a/Engine/Source/Runtime/NetworkReplayStreaming/HttpNetworkReplayStreaming/Private/HttpNetworkReplayStreaming.cpp b/Engine/Source/Runtime/NetworkReplayStreaming/HttpNetworkReplayStreaming/Private/HttpNetworkReplayStreaming.cpp
--- a/Engine/Source/Runtime/NetworkReplayStreaming/HttpNetworkReplayStreaming/Private/HttpNetworkReplayStreaming.cpp
+++ b/Engine/Source/Runtime/NetworkReplayStreaming/HttpNetworkReplayStreaming/Private/HttpNetworkReplayStreaming.cpp
@@ -18,6 +18,7 @@ public:
 	int32		DemoTimeInMs;
 	int32		NumViewers;
 	bool		bIsLive;
+	int32		Changelist;
 
 	// FOnlineJsonSerializable
 	BEGIN_ONLINE_JSON_SERIALIZER
@@ -29,6 +30,7 @@ public:
 		ONLINE_JSON_SERIALIZE( "DemoTimeInMs",	DemoTimeInMs );
 		ONLINE_JSON_SERIALIZE( "NumViewers",	NumViewers );
 		ONLINE_JSON_SERIALIZE( "bIsLive",		bIsLive );
+		ONLINE_JSON_SERIALIZE( "Changelist",	Changelist );
 	END_ONLINE_JSON_SERIALIZER
 };
 
@@ -1099,7 +1101,8 @@ void FHttpNetworkReplayStreamer::HttpEnumerateSessionsFinished( FHttpRequestPtr
 			NewStream.SizeInBytes	= ReplayList.Replays[i].SizeInBytes;
 			NewStream.LengthInMS	= ReplayList.Replays[i].DemoTimeInMs;
 			NewStream.NumViewers	= ReplayList.Replays[i].NumViewers;
-			NewStream.bIsLive		= ReplayList.Replays[i].bIsLive;;
+			NewStream.bIsLive		= ReplayList.Replays[i].bIsLive;
+			NewStream.Changelist	= ReplayList.Replays[i].Changelist;
 
 			Streams.Add( NewStream );
 		}
diff --git a/Engine/Source/Runtime/NetworkReplayStreaming/NetworkReplayStreaming/Public/NetworkReplayStreaming.h b/Engine/Source/Runtime/NetworkReplayStreaming/NetworkReplayStreaming/Public/NetworkReplayStreaming.h
--- a/Engine/Source/Runtime/NetworkReplayStreaming/NetworkReplayStreaming/Public/NetworkReplayStreaming.h
+++ b/Engine/Source/Runtime/NetworkReplayStreaming/NetworkReplayStreaming/Public/NetworkReplayStreaming.h
@@ -9,7 +9,7 @@
 /** Struct to store information about a stream, returned from search results. */
 struct FNetworkReplayStreamInfo
 {
-	FNetworkReplayStreamInfo() : SizeInBytes( 0 ), LengthInMS( 0 ), NumViewers( 0 ), bIsLive( false ) {}
+	FNetworkReplayStreamInfo() : SizeInBytes( 0 ), LengthInMS( 0 ), NumViewers( 0 ), bIsLive( false ), Changelist( 0 ) {}
 
 	/** The name of the stream (generally this is auto generated, refer to friendly name for UI) */
 	FString Name;
@@ -31,6 +31,9 @@ struct FNetworkReplayStreamInfo
 
 	/** True if the stream is live and the game hasn't completed yet */
 	bool bIsLive;
+
+	/** The changelist of the replay */
+	int32 Changelist;
 };
 
 namespace ENetworkReplayError
