Fix issue when there is an http failure with a StopStreamingCall already in flight

diff --git a/Engine/Source/Runtime/NetworkReplayStreaming/HttpNetworkReplayStreaming/Private/HttpNetworkReplayStreaming.cpp b/Engine/Source/Runtime/NetworkReplayStreaming/HttpNetworkReplayStreaming/Private/HttpNetworkReplayStreaming.cpp
--- a/Engine/Source/Runtime/NetworkReplayStreaming/HttpNetworkReplayStreaming/Private/HttpNetworkReplayStreaming.cpp
+++ b/Engine/Source/Runtime/NetworkReplayStreaming/HttpNetworkReplayStreaming/Private/HttpNetworkReplayStreaming.cpp
@@ -615,8 +615,9 @@ void FHttpNetworkReplayStreamer::SetLastError( const ENetworkReplayError::Type I
 	TSharedPtr< FQueuedHttpRequest > QueuedRequest;
 	while ( QueuedHttpRequests.Dequeue( QueuedRequest ) ) { }
 
-	StreamerState		= EStreamerState::Idle;
-	StreamerLastError	= InLastError;
+	StreamerState			= EStreamerState::Idle;
+	StreamerLastError		= InLastError;
+	bStopStreamingCalled	= false;
 }
 
 ENetworkReplayError::Type FHttpNetworkReplayStreamer::GetLastError() const
