Engine side work to support showing recently viewed replays

diff --git a/Engine/Source/Runtime/Engine/Private/DemoNetDriver.cpp b/Engine/Source/Runtime/Engine/Private/DemoNetDriver.cpp
--- a/Engine/Source/Runtime/Engine/Private/DemoNetDriver.cpp
+++ b/Engine/Source/Runtime/Engine/Private/DemoNetDriver.cpp
@@ -227,9 +227,27 @@ bool UDemoNetDriver::InitConnect( FNetworkNotify* InNotify, const FURL& ConnectU
 	ServerConnection = NewObject<UNetConnection>(GetTransientPackage(), UDemoNetConnection::StaticClass());
 	ServerConnection->InitConnection( this, USOCK_Pending, ConnectURL, 1000000 );
 
+	FString UserName;
+
+	if ( GetWorld()->GetGameInstance()->GetFirstGamePlayer() != nullptr )
+	{
+		TSharedPtr<const FUniqueNetId> ViewerId = GetWorld()->GetGameInstance()->GetFirstGamePlayer()->GetPreferredUniqueNetId();
+
+		if ( ViewerId.IsValid() )
+		{ 
+			UserName = ViewerId->ToString();
+		}
+	}
+
 	bWasStartStreamingSuccessful = true;
-	// Friendly name isn't important for loading an existing replay.
-	ReplayStreamer->StartStreaming( DemoFilename, FString(), false, FNetworkVersion::GetReplayVersion(), FOnStreamReadyDelegate::CreateUObject( this, &UDemoNetDriver::ReplayStreamingReady ) );
+
+	ReplayStreamer->StartStreaming( 
+		DemoFilename, 
+		FString(),		// Friendly name isn't important for loading an existing replay.
+		UserName, 
+		false, 
+		FNetworkVersion::GetReplayVersion(), 
+		FOnStreamReadyDelegate::CreateUObject( this, &UDemoNetDriver::ReplayStreamingReady ) );
 
 	return bWasStartStreamingSuccessful;
 }
@@ -399,6 +417,7 @@ bool UDemoNetDriver::InitListen( FNetworkNotify* InNotify, FURL& ListenURL, bool
 	ReplayStreamer->StartStreaming(
 		DemoFilename,
 		FriendlyNameOption != nullptr ? FString(FriendlyNameOption) : World->GetMapName(),
+		FString(),		// Don't need user name for upload currently
 		true,
 		FNetworkVersion::GetReplayVersion(),
 		FOnStreamReadyDelegate::CreateUObject( this, &UDemoNetDriver::ReplayStreamingReady ) );
diff --git a/Engine/Source/Runtime/NetworkReplayStreaming/HttpNetworkReplayStreaming/Private/HttpNetworkReplayStreaming.cpp b/Engine/Source/Runtime/NetworkReplayStreaming/HttpNetworkReplayStreaming/Private/HttpNetworkReplayStreaming.cpp
--- a/Engine/Source/Runtime/NetworkReplayStreaming/HttpNetworkReplayStreaming/Private/HttpNetworkReplayStreaming.cpp
+++ b/Engine/Source/Runtime/NetworkReplayStreaming/HttpNetworkReplayStreaming/Private/HttpNetworkReplayStreaming.cpp
@@ -124,7 +124,7 @@ FHttpNetworkReplayStreamer::FHttpNetworkReplayStreamer() :
 	GConfig->GetString( TEXT( "HttpNetworkReplayStreaming" ), TEXT( "ServerURL" ), ServerURL, GEngineIni );
 }
 
-void FHttpNetworkReplayStreamer::StartStreaming( const FString& CustomName, const FString& FriendlyName, bool bRecord, const FNetworkReplayVersion& InReplayVersion, const FOnStreamReadyDelegate& Delegate )
+void FHttpNetworkReplayStreamer::StartStreaming( const FString& CustomName, const FString& FriendlyName, const FString& UserName, bool bRecord, const FNetworkReplayVersion& InReplayVersion, const FOnStreamReadyDelegate& Delegate )
 {
 	if ( !SessionName.IsEmpty() )
 	{
@@ -180,7 +180,7 @@ void FHttpNetworkReplayStreamer::StartStreaming( const FString& CustomName, cons
 		SessionName = CustomName;
 
 		// Notify the http server that we want to start downloading a replay
-		HttpRequest->SetURL( FString::Printf( TEXT( "%sstartdownloading?Session=%s" ), *ServerURL, *SessionName ) );
+		HttpRequest->SetURL( FString::Printf( TEXT( "%sstartdownloading?Session=%s&User=%s" ), *ServerURL, *SessionName, *UserName ) );
 		HttpRequest->SetVerb( TEXT( "POST" ) );
 
 		HttpRequest->OnProcessRequestComplete().BindRaw( this, &FHttpNetworkReplayStreamer::HttpStartDownloadingFinished );
@@ -742,6 +742,27 @@ void FHttpNetworkReplayStreamer::EnumerateStreams( const FNetworkReplayVersion&
 	AddRequestToQueue( EQueuedHttpRequestType::EnumeratingSessions, HttpRequest );
 }
 
+void FHttpNetworkReplayStreamer::EnumerateRecentStreams( const FNetworkReplayVersion& InReplayVersion, const FString& InRecentViewer, const FOnEnumerateStreamsComplete& Delegate )
+{
+	if ( IsHttpRequestInFlight() )
+	{
+		Delegate.ExecuteIfBound( TArray<FNetworkReplayStreamInfo>() );
+		return;
+	}
+
+	TSharedRef<class IHttpRequest> HttpRequest = FHttpModule::Get().CreateRequest();
+
+	// Enumerate all of the sessions
+	HttpRequest->SetURL( FString::Printf( TEXT( "%senumsessions?App=%s&Version=%u&CL=%u&User=%s" ), *ServerURL, *InReplayVersion.AppString, InReplayVersion.NetworkVersion, InReplayVersion.Changelist, *InRecentViewer ) );
+	HttpRequest->SetVerb( TEXT( "GET" ) );
+
+	HttpRequest->OnProcessRequestComplete().BindRaw( this, &FHttpNetworkReplayStreamer::HttpEnumerateSessionsFinished );
+
+	EnumerateStreamsDelegate = Delegate;
+
+	AddRequestToQueue( EQueuedHttpRequestType::EnumeratingSessions, HttpRequest );
+}
+
 void FHttpNetworkReplayStreamer::EnumerateCheckpoints()
 {
 	TSharedRef<class IHttpRequest> HttpRequest = FHttpModule::Get().CreateRequest();
diff --git a/Engine/Source/Runtime/NetworkReplayStreaming/HttpNetworkReplayStreaming/Private/HttpNetworkReplayStreaming.h b/Engine/Source/Runtime/NetworkReplayStreaming/HttpNetworkReplayStreaming/Private/HttpNetworkReplayStreaming.h
--- a/Engine/Source/Runtime/NetworkReplayStreaming/HttpNetworkReplayStreaming/Private/HttpNetworkReplayStreaming.h
+++ b/Engine/Source/Runtime/NetworkReplayStreaming/HttpNetworkReplayStreaming/Private/HttpNetworkReplayStreaming.h
@@ -96,7 +96,7 @@ public:
 	FHttpNetworkReplayStreamer();
 
 	/** INetworkReplayStreamer implementation */
-	virtual void		StartStreaming( const FString& CustomName, const FString& FriendlyName, bool bRecord, const FNetworkReplayVersion& ReplayVersion, const FOnStreamReadyDelegate& Delegate ) override;
+	virtual void		StartStreaming( const FString& CustomName, const FString& FriendlyName, const FString& UserName, bool bRecord, const FNetworkReplayVersion& ReplayVersion, const FOnStreamReadyDelegate& Delegate ) override;
 	virtual void		StopStreaming() override;
 	virtual FArchive*	GetHeaderArchive() override;
 	virtual FArchive*	GetStreamingArchive() override;
@@ -114,6 +114,8 @@ public:
 	virtual bool		IsLive() const override;
 	virtual void		DeleteFinishedStream( const FString& StreamName, const FOnDeleteFinishedStreamComplete& Delegate ) const override;
 	virtual void		EnumerateStreams( const FNetworkReplayVersion& ReplayVersion, const FOnEnumerateStreamsComplete& Delegate ) override;
+	virtual void		EnumerateRecentStreams( const FNetworkReplayVersion& ReplayVersion, const FString& RecentViewer, const FOnEnumerateStreamsComplete& Delegate ) override;
+
 	virtual ENetworkReplayError::Type GetLastError() const override;
 
 	/** FHttpNetworkReplayStreamer */
diff --git a/Engine/Source/Runtime/NetworkReplayStreaming/NetworkReplayStreaming/Public/NetworkReplayStreaming.h b/Engine/Source/Runtime/NetworkReplayStreaming/NetworkReplayStreaming/Public/NetworkReplayStreaming.h
--- a/Engine/Source/Runtime/NetworkReplayStreaming/NetworkReplayStreaming/Public/NetworkReplayStreaming.h
+++ b/Engine/Source/Runtime/NetworkReplayStreaming/NetworkReplayStreaming/Public/NetworkReplayStreaming.h
@@ -106,7 +106,7 @@ class INetworkReplayStreamer
 public:
 	virtual ~INetworkReplayStreamer() {}
 
-	virtual void StartStreaming( const FString& CustomName, const FString& FriendlyName, bool bRecord, const FNetworkReplayVersion& ReplayVersion, const FOnStreamReadyDelegate& Delegate ) = 0;
+	virtual void StartStreaming( const FString& CustomName, const FString& FriendlyName, const FString& UserName, bool bRecord, const FNetworkReplayVersion& ReplayVersion, const FOnStreamReadyDelegate& Delegate ) = 0;
 	virtual void StopStreaming() = 0;
 	virtual FArchive* GetHeaderArchive() = 0;
 	virtual FArchive* GetStreamingArchive() = 0;
@@ -140,6 +140,13 @@ public:
 	 */
 	virtual void EnumerateStreams( const FNetworkReplayVersion& ReplayVersion, const FOnEnumerateStreamsComplete& Delegate ) = 0;
 
+	/**
+	 * Retrieves the streams that have been recently viewed. May execute asynchronously.
+	 *
+	 * @param Delegate A delegate that will be executed if bound when the list of streams is available
+	 */
+	virtual void EnumerateRecentStreams( const FNetworkReplayVersion& ReplayVersion, const FString& RecentViewer, const FOnEnumerateStreamsComplete& Delegate ) = 0;
+
 	/** Returns the last error that occurred while streaming replays */
 	virtual ENetworkReplayError::Type GetLastError() const = 0;
 };
diff --git a/Engine/Source/Runtime/NetworkReplayStreaming/NullNetworkReplayStreaming/Private/NullNetworkReplayStreaming.cpp b/Engine/Source/Runtime/NetworkReplayStreaming/NullNetworkReplayStreaming/Private/NullNetworkReplayStreaming.cpp
--- a/Engine/Source/Runtime/NetworkReplayStreaming/NullNetworkReplayStreaming/Private/NullNetworkReplayStreaming.cpp
+++ b/Engine/Source/Runtime/NetworkReplayStreaming/NullNetworkReplayStreaming/Private/NullNetworkReplayStreaming.cpp
@@ -158,7 +158,7 @@ static FString GetAutomaticDemoName()
 	return FinalDemoName;
 }
 
-void FNullNetworkReplayStreamer::StartStreaming( const FString& CustomName, const FString& FriendlyName, bool bRecord, const FNetworkReplayVersion& ReplayVersion, const FOnStreamReadyDelegate& Delegate )
+void FNullNetworkReplayStreamer::StartStreaming( const FString& CustomName, const FString& FriendlyName, const FString& UserName, bool bRecord, const FNetworkReplayVersion& ReplayVersion, const FOnStreamReadyDelegate& Delegate )
 {
 	FString FinalDemoName = CustomName;
 
diff --git a/Engine/Source/Runtime/NetworkReplayStreaming/NullNetworkReplayStreaming/Public/NullNetworkReplayStreaming.h b/Engine/Source/Runtime/NetworkReplayStreaming/NullNetworkReplayStreaming/Public/NullNetworkReplayStreaming.h
--- a/Engine/Source/Runtime/NetworkReplayStreaming/NullNetworkReplayStreaming/Public/NullNetworkReplayStreaming.h
+++ b/Engine/Source/Runtime/NetworkReplayStreaming/NullNetworkReplayStreaming/Public/NullNetworkReplayStreaming.h
@@ -38,7 +38,7 @@ public:
 	{}
 	
 	/** INetworkReplayStreamer implementation */
-	virtual void StartStreaming( const FString& CustomName, const FString& FriendlyName, bool bRecord, const FNetworkReplayVersion& ReplayVersion, const FOnStreamReadyDelegate& Delegate ) override;
+	virtual void StartStreaming( const FString& CustomName, const FString& FriendlyName, const FString& UserName, bool bRecord, const FNetworkReplayVersion& ReplayVersion, const FOnStreamReadyDelegate& Delegate ) override;
 	virtual void StopStreaming() override;
 	virtual FArchive* GetHeaderArchive() override;
 	virtual FArchive* GetStreamingArchive() override;
@@ -56,6 +56,7 @@ public:
 	virtual bool IsLive() const override;
 	virtual void DeleteFinishedStream( const FString& StreamName, const FOnDeleteFinishedStreamComplete& Delegate) const override;
 	virtual void EnumerateStreams( const FNetworkReplayVersion& ReplayVersion, const FOnEnumerateStreamsComplete& Delegate ) override;
+	virtual void EnumerateRecentStreams( const FNetworkReplayVersion& ReplayVersion, const FString& RecentViewer, const FOnEnumerateStreamsComplete& Delegate ) override {}
 	virtual ENetworkReplayError::Type GetLastError() const override { return ENetworkReplayError::None; }
 
 private:
