Engine side support for replay meta tags

diff --git a/Engine/Source/Runtime/NetworkReplayStreaming/HttpNetworkReplayStreaming/Private/HttpNetworkReplayStreaming.cpp b/Engine/Source/Runtime/NetworkReplayStreaming/HttpNetworkReplayStreaming/Private/HttpNetworkReplayStreaming.cpp
--- a/Engine/Source/Runtime/NetworkReplayStreaming/HttpNetworkReplayStreaming/Private/HttpNetworkReplayStreaming.cpp
+++ b/Engine/Source/Runtime/NetworkReplayStreaming/HttpNetworkReplayStreaming/Private/HttpNetworkReplayStreaming.cpp
@@ -204,8 +204,18 @@ void FHttpNetworkReplayStreamer::StartStreaming( const FString& CustomName, cons
 		// We can't upload the header until we have the session name, which depends on the StartUploading task above to finish
 		bNeedToUploadHeader = true;
 
-		// Notify the http server that we want to start uploading a replay
-		HttpRequest->SetURL( FString::Printf( TEXT( "%sstartuploading?App=%s&Version=%u&CL=%u&Friendly=%s" ), *ServerURL, *ReplayVersion.AppString, ReplayVersion.NetworkVersion, ReplayVersion.Changelist, *FriendlyName ) );
+		FString MetaString;
+
+		if ( FParse::Value( FCommandLine::Get(), TEXT( "ReplayMeta=" ), MetaString ) && !MetaString.IsEmpty() )
+		{
+			// Notify the http server that we want to start uploading a replay
+			HttpRequest->SetURL( FString::Printf( TEXT( "%sstartuploading?App=%s&Version=%u&CL=%u&Friendly=%s&Meta=%s" ), *ServerURL, *ReplayVersion.AppString, ReplayVersion.NetworkVersion, ReplayVersion.Changelist, *FriendlyName, *MetaString ) );
+		}
+		else
+		{
+			// Notify the http server that we want to start uploading a replay
+			HttpRequest->SetURL( FString::Printf( TEXT( "%sstartuploading?App=%s&Version=%u&CL=%u&Friendly=%s" ), *ServerURL, *ReplayVersion.AppString, ReplayVersion.NetworkVersion, ReplayVersion.Changelist, *FriendlyName ) );
+		}
 
 		HttpRequest->SetVerb( TEXT( "POST" ) );
 
@@ -721,7 +731,7 @@ void FHttpNetworkReplayStreamer::DeleteFinishedStream( const FString& StreamName
 	Delegate.ExecuteIfBound(false);
 }
 
-void FHttpNetworkReplayStreamer::EnumerateStreams( const FNetworkReplayVersion& InReplayVersion, const FOnEnumerateStreamsComplete& Delegate )
+void FHttpNetworkReplayStreamer::EnumerateStreams( const FNetworkReplayVersion& InReplayVersion, const FString& MetaString, const FOnEnumerateStreamsComplete& Delegate )
 {
 	if ( IsHttpRequestInFlight() )
 	{
@@ -732,7 +742,15 @@ void FHttpNetworkReplayStreamer::EnumerateStreams( const FNetworkReplayVersion&
 	TSharedRef<class IHttpRequest> HttpRequest = FHttpModule::Get().CreateRequest();
 
 	// Enumerate all of the sessions
-	HttpRequest->SetURL( FString::Printf( TEXT( "%senumsessions?App=%s&Version=%u&CL=%u" ), *ServerURL, *InReplayVersion.AppString, InReplayVersion.NetworkVersion, InReplayVersion.Changelist ) );
+	if ( !MetaString.IsEmpty() )
+	{
+		HttpRequest->SetURL( FString::Printf( TEXT( "%senumsessions?App=%s&Version=%u&CL=%u&Meta=%s" ), *ServerURL, *InReplayVersion.AppString, InReplayVersion.NetworkVersion, InReplayVersion.Changelist, *MetaString ) );
+	}
+	else
+	{
+		HttpRequest->SetURL( FString::Printf( TEXT( "%senumsessions?App=%s&Version=%u&CL=%u" ), *ServerURL, *InReplayVersion.AppString, InReplayVersion.NetworkVersion, InReplayVersion.Changelist ) );
+	}
+
 	HttpRequest->SetVerb( TEXT( "GET" ) );
 
 	HttpRequest->OnProcessRequestComplete().BindRaw( this, &FHttpNetworkReplayStreamer::HttpEnumerateSessionsFinished );
diff --git a/Engine/Source/Runtime/NetworkReplayStreaming/HttpNetworkReplayStreaming/Private/HttpNetworkReplayStreaming.h b/Engine/Source/Runtime/NetworkReplayStreaming/HttpNetworkReplayStreaming/Private/HttpNetworkReplayStreaming.h
--- a/Engine/Source/Runtime/NetworkReplayStreaming/HttpNetworkReplayStreaming/Private/HttpNetworkReplayStreaming.h
+++ b/Engine/Source/Runtime/NetworkReplayStreaming/HttpNetworkReplayStreaming/Private/HttpNetworkReplayStreaming.h
@@ -113,7 +113,7 @@ public:
 	virtual bool		IsLoadingCheckpoint() const override;
 	virtual bool		IsLive() const override;
 	virtual void		DeleteFinishedStream( const FString& StreamName, const FOnDeleteFinishedStreamComplete& Delegate ) const override;
-	virtual void		EnumerateStreams( const FNetworkReplayVersion& ReplayVersion, const FOnEnumerateStreamsComplete& Delegate ) override;
+	virtual void		EnumerateStreams( const FNetworkReplayVersion& ReplayVersion, const FString& MetaString, const FOnEnumerateStreamsComplete& Delegate ) override;
 	virtual void		EnumerateRecentStreams( const FNetworkReplayVersion& ReplayVersion, const FString& RecentViewer, const FOnEnumerateStreamsComplete& Delegate ) override;
 
 	virtual ENetworkReplayError::Type GetLastError() const override;
diff --git a/Engine/Source/Runtime/NetworkReplayStreaming/NetworkReplayStreaming/Public/NetworkReplayStreaming.h b/Engine/Source/Runtime/NetworkReplayStreaming/NetworkReplayStreaming/Public/NetworkReplayStreaming.h
--- a/Engine/Source/Runtime/NetworkReplayStreaming/NetworkReplayStreaming/Public/NetworkReplayStreaming.h
+++ b/Engine/Source/Runtime/NetworkReplayStreaming/NetworkReplayStreaming/Public/NetworkReplayStreaming.h
@@ -138,7 +138,7 @@ public:
 	 *
 	 * @param Delegate A delegate that will be executed if bound when the list of streams is available
 	 */
-	virtual void EnumerateStreams( const FNetworkReplayVersion& ReplayVersion, const FOnEnumerateStreamsComplete& Delegate ) = 0;
+	virtual void EnumerateStreams( const FNetworkReplayVersion& ReplayVersion, const FString& MetaString, const FOnEnumerateStreamsComplete& Delegate ) = 0;
 
 	/**
 	 * Retrieves the streams that have been recently viewed. May execute asynchronously.
diff --git a/Engine/Source/Runtime/NetworkReplayStreaming/NullNetworkReplayStreaming/Private/NullNetworkReplayStreaming.cpp b/Engine/Source/Runtime/NetworkReplayStreaming/NullNetworkReplayStreaming/Private/NullNetworkReplayStreaming.cpp
--- a/Engine/Source/Runtime/NetworkReplayStreaming/NullNetworkReplayStreaming/Private/NullNetworkReplayStreaming.cpp
+++ b/Engine/Source/Runtime/NetworkReplayStreaming/NullNetworkReplayStreaming/Private/NullNetworkReplayStreaming.cpp
@@ -317,7 +317,7 @@ void FNullNetworkReplayStreamer::DeleteFinishedStream( const FString& StreamName
 	Delegate.ExecuteIfBound(DeleteSucceeded);
 }
 
-void FNullNetworkReplayStreamer::EnumerateStreams( const FNetworkReplayVersion& ReplayVersion, const FOnEnumerateStreamsComplete& Delegate )
+void FNullNetworkReplayStreamer::EnumerateStreams( const FNetworkReplayVersion& ReplayVersion, const FString& MetaString, const FOnEnumerateStreamsComplete& Delegate )
 {
 	// Simply returns a stream for each folder in the Saved/Demos directory
 	const FString WildCardPath = GetDemoPath() + TEXT( "*" );
diff --git a/Engine/Source/Runtime/NetworkReplayStreaming/NullNetworkReplayStreaming/Public/NullNetworkReplayStreaming.h b/Engine/Source/Runtime/NetworkReplayStreaming/NullNetworkReplayStreaming/Public/NullNetworkReplayStreaming.h
--- a/Engine/Source/Runtime/NetworkReplayStreaming/NullNetworkReplayStreaming/Public/NullNetworkReplayStreaming.h
+++ b/Engine/Source/Runtime/NetworkReplayStreaming/NullNetworkReplayStreaming/Public/NullNetworkReplayStreaming.h
@@ -55,7 +55,7 @@ public:
 	virtual bool IsLoadingCheckpoint() const override { return false; }
 	virtual bool IsLive() const override;
 	virtual void DeleteFinishedStream( const FString& StreamName, const FOnDeleteFinishedStreamComplete& Delegate) const override;
-	virtual void EnumerateStreams( const FNetworkReplayVersion& ReplayVersion, const FOnEnumerateStreamsComplete& Delegate ) override;
+	virtual void EnumerateStreams( const FNetworkReplayVersion& ReplayVersion, const FString& MetaString, const FOnEnumerateStreamsComplete& Delegate ) override;
 	virtual void EnumerateRecentStreams( const FNetworkReplayVersion& ReplayVersion, const FString& RecentViewer, const FOnEnumerateStreamsComplete& Delegate ) override {}
 	virtual ENetworkReplayError::Type GetLastError() const override { return ENetworkReplayError::None; }
 
