Add engine support to add participating users to replays

diff --git a/Engine/Source/Runtime/Engine/Private/DemoNetDriver.cpp b/Engine/Source/Runtime/Engine/Private/DemoNetDriver.cpp
--- a/Engine/Source/Runtime/Engine/Private/DemoNetDriver.cpp
+++ b/Engine/Source/Runtime/Engine/Private/DemoNetDriver.cpp
@@ -227,7 +227,7 @@ bool UDemoNetDriver::InitConnect( FNetworkNotify* InNotify, const FURL& ConnectU
 	ServerConnection = NewObject<UNetConnection>(GetTransientPackage(), UDemoNetConnection::StaticClass());
 	ServerConnection->InitConnection( this, USOCK_Pending, ConnectURL, 1000000 );
 
-	FString UserName;
+	TArray< FString > UserNames;
 
 	if ( GetWorld()->GetGameInstance()->GetFirstGamePlayer() != nullptr )
 	{
@@ -235,7 +235,7 @@ bool UDemoNetDriver::InitConnect( FNetworkNotify* InNotify, const FURL& ConnectU
 
 		if ( ViewerId.IsValid() )
 		{ 
-			UserName = ViewerId->ToString();
+			UserNames.Add( ViewerId->ToString() );
 		}
 	}
 
@@ -244,7 +244,7 @@ bool UDemoNetDriver::InitConnect( FNetworkNotify* InNotify, const FURL& ConnectU
 	ReplayStreamer->StartStreaming( 
 		DemoFilename, 
 		FString(),		// Friendly name isn't important for loading an existing replay.
-		UserName, 
+		UserNames, 
 		false, 
 		FNetworkVersion::GetReplayVersion(), 
 		FOnStreamReadyDelegate::CreateUObject( this, &UDemoNetDriver::ReplayStreamingReady ) );
@@ -414,10 +414,17 @@ bool UDemoNetDriver::InitListen( FNetworkNotify* InNotify, FURL& ListenURL, bool
 
 	const TCHAR* FriendlyNameOption = ListenURL.GetOption( TEXT("DemoFriendlyName="), nullptr );
 
+	TArray< FString > UserNames;
+
+	for ( int32 i = 0; i < GetWorld()->GameState->PlayerArray.Num(); i++ )
+	{
+		UserNames.Add( GetWorld()->GameState->PlayerArray[i]->UniqueId.ToString() );
+	}
+
 	ReplayStreamer->StartStreaming(
 		DemoFilename,
 		FriendlyNameOption != nullptr ? FString(FriendlyNameOption) : World->GetMapName(),
-		FString(),		// Don't need user name for upload currently
+		UserNames,
 		true,
 		FNetworkVersion::GetReplayVersion(),
 		FOnStreamReadyDelegate::CreateUObject( this, &UDemoNetDriver::ReplayStreamingReady ) );
diff --git a/Engine/Source/Runtime/NetworkReplayStreaming/HttpNetworkReplayStreaming/Private/HttpNetworkReplayStreaming.cpp b/Engine/Source/Runtime/NetworkReplayStreaming/HttpNetworkReplayStreaming/Private/HttpNetworkReplayStreaming.cpp
--- a/Engine/Source/Runtime/NetworkReplayStreaming/HttpNetworkReplayStreaming/Private/HttpNetworkReplayStreaming.cpp
+++ b/Engine/Source/Runtime/NetworkReplayStreaming/HttpNetworkReplayStreaming/Private/HttpNetworkReplayStreaming.cpp
@@ -49,6 +49,24 @@ public:
 	END_ONLINE_JSON_SERIALIZER
 };
 
+class FNetworkReplayUserList : public FOnlineJsonSerializable
+{
+public:
+	FNetworkReplayUserList()
+	{
+	}
+	virtual ~FNetworkReplayUserList()
+	{
+	}
+
+	TArray< FString > Users;
+
+	// FOnlineJsonSerializable
+	BEGIN_ONLINE_JSON_SERIALIZER
+		ONLINE_JSON_SERIALIZE_ARRAY( "users", Users );
+	END_ONLINE_JSON_SERIALIZER
+};
+
 void FHttpStreamFArchive::Serialize( void* V, int64 Length ) 
 {
 	if ( IsLoading() )
@@ -124,7 +142,7 @@ FHttpNetworkReplayStreamer::FHttpNetworkReplayStreamer() :
 	GConfig->GetString( TEXT( "HttpNetworkReplayStreaming" ), TEXT( "ServerURL" ), ServerURL, GEngineIni );
 }
 
-void FHttpNetworkReplayStreamer::StartStreaming( const FString& CustomName, const FString& FriendlyName, const FString& UserName, bool bRecord, const FNetworkReplayVersion& InReplayVersion, const FOnStreamReadyDelegate& Delegate )
+void FHttpNetworkReplayStreamer::StartStreaming( const FString& CustomName, const FString& FriendlyName, const TArray< FString >& UserNames, bool bRecord, const FNetworkReplayVersion& InReplayVersion, const FOnStreamReadyDelegate& Delegate )
 {
 	if ( !SessionName.IsEmpty() )
 	{
@@ -179,6 +197,13 @@ void FHttpNetworkReplayStreamer::StartStreaming( const FString& CustomName, cons
 
 		SessionName = CustomName;
 
+		FString UserName;
+
+		if ( UserNames.Num() == 1 )
+		{
+			UserName = UserNames[0];
+		}
+
 		// Notify the http server that we want to start downloading a replay
 		HttpRequest->SetURL( FString::Printf( TEXT( "%sstartdownloading?Session=%s&User=%s" ), *ServerURL, *SessionName, *UserName ) );
 		HttpRequest->SetVerb( TEXT( "POST" ) );
@@ -221,6 +246,15 @@ void FHttpNetworkReplayStreamer::StartStreaming( const FString& CustomName, cons
 
 		HttpRequest->OnProcessRequestComplete().BindRaw( this, &FHttpNetworkReplayStreamer::HttpStartUploadingFinished );
 
+		if ( UserNames.Num() > 0 )
+		{
+			FNetworkReplayUserList UserList;
+
+			UserList.Users = UserNames;
+			HttpRequest->SetContentAsString( UserList.ToJson() );
+			HttpRequest->SetHeader( TEXT( "Content-Type" ), TEXT( "application/json" ) );
+		}
+
 		AddRequestToQueue( EQueuedHttpRequestType::StartUploading, HttpRequest );
 	}
 }
diff --git a/Engine/Source/Runtime/NetworkReplayStreaming/HttpNetworkReplayStreaming/Private/HttpNetworkReplayStreaming.h b/Engine/Source/Runtime/NetworkReplayStreaming/HttpNetworkReplayStreaming/Private/HttpNetworkReplayStreaming.h
--- a/Engine/Source/Runtime/NetworkReplayStreaming/HttpNetworkReplayStreaming/Private/HttpNetworkReplayStreaming.h
+++ b/Engine/Source/Runtime/NetworkReplayStreaming/HttpNetworkReplayStreaming/Private/HttpNetworkReplayStreaming.h
@@ -96,7 +96,7 @@ public:
 	FHttpNetworkReplayStreamer();
 
 	/** INetworkReplayStreamer implementation */
-	virtual void		StartStreaming( const FString& CustomName, const FString& FriendlyName, const FString& UserName, bool bRecord, const FNetworkReplayVersion& ReplayVersion, const FOnStreamReadyDelegate& Delegate ) override;
+	virtual void		StartStreaming( const FString& CustomName, const FString& FriendlyName, const TArray< FString >& UserNames, bool bRecord, const FNetworkReplayVersion& ReplayVersion, const FOnStreamReadyDelegate& Delegate ) override;
 	virtual void		StopStreaming() override;
 	virtual FArchive*	GetHeaderArchive() override;
 	virtual FArchive*	GetStreamingArchive() override;
diff --git a/Engine/Source/Runtime/NetworkReplayStreaming/NetworkReplayStreaming/Public/NetworkReplayStreaming.h b/Engine/Source/Runtime/NetworkReplayStreaming/NetworkReplayStreaming/Public/NetworkReplayStreaming.h
--- a/Engine/Source/Runtime/NetworkReplayStreaming/NetworkReplayStreaming/Public/NetworkReplayStreaming.h
+++ b/Engine/Source/Runtime/NetworkReplayStreaming/NetworkReplayStreaming/Public/NetworkReplayStreaming.h
@@ -106,7 +106,7 @@ class INetworkReplayStreamer
 public:
 	virtual ~INetworkReplayStreamer() {}
 
-	virtual void StartStreaming( const FString& CustomName, const FString& FriendlyName, const FString& UserName, bool bRecord, const FNetworkReplayVersion& ReplayVersion, const FOnStreamReadyDelegate& Delegate ) = 0;
+	virtual void StartStreaming( const FString& CustomName, const FString& FriendlyName, const TArray< FString >& UserNames, bool bRecord, const FNetworkReplayVersion& ReplayVersion, const FOnStreamReadyDelegate& Delegate ) = 0;
 	virtual void StopStreaming() = 0;
 	virtual FArchive* GetHeaderArchive() = 0;
 	virtual FArchive* GetStreamingArchive() = 0;
diff --git a/Engine/Source/Runtime/NetworkReplayStreaming/NullNetworkReplayStreaming/Private/NullNetworkReplayStreaming.cpp b/Engine/Source/Runtime/NetworkReplayStreaming/NullNetworkReplayStreaming/Private/NullNetworkReplayStreaming.cpp
--- a/Engine/Source/Runtime/NetworkReplayStreaming/NullNetworkReplayStreaming/Private/NullNetworkReplayStreaming.cpp
+++ b/Engine/Source/Runtime/NetworkReplayStreaming/NullNetworkReplayStreaming/Private/NullNetworkReplayStreaming.cpp
@@ -158,7 +158,7 @@ static FString GetAutomaticDemoName()
 	return FinalDemoName;
 }
 
-void FNullNetworkReplayStreamer::StartStreaming( const FString& CustomName, const FString& FriendlyName, const FString& UserName, bool bRecord, const FNetworkReplayVersion& ReplayVersion, const FOnStreamReadyDelegate& Delegate )
+void FNullNetworkReplayStreamer::StartStreaming( const FString& CustomName, const FString& FriendlyName, const TArray< FString >& UserNames, bool bRecord, const FNetworkReplayVersion& ReplayVersion, const FOnStreamReadyDelegate& Delegate )
 {
 	FString FinalDemoName = CustomName;
 
diff --git a/Engine/Source/Runtime/NetworkReplayStreaming/NullNetworkReplayStreaming/Public/NullNetworkReplayStreaming.h b/Engine/Source/Runtime/NetworkReplayStreaming/NullNetworkReplayStreaming/Public/NullNetworkReplayStreaming.h
--- a/Engine/Source/Runtime/NetworkReplayStreaming/NullNetworkReplayStreaming/Public/NullNetworkReplayStreaming.h
+++ b/Engine/Source/Runtime/NetworkReplayStreaming/NullNetworkReplayStreaming/Public/NullNetworkReplayStreaming.h
@@ -38,7 +38,7 @@ public:
 	{}
 	
 	/** INetworkReplayStreamer implementation */
-	virtual void StartStreaming( const FString& CustomName, const FString& FriendlyName, const FString& UserName, bool bRecord, const FNetworkReplayVersion& ReplayVersion, const FOnStreamReadyDelegate& Delegate ) override;
+	virtual void StartStreaming( const FString& CustomName, const FString& FriendlyName, const TArray< FString >& UserNames, bool bRecord, const FNetworkReplayVersion& ReplayVersion, const FOnStreamReadyDelegate& Delegate ) override;
 	virtual void StopStreaming() override;
 	virtual FArchive* GetHeaderArchive() override;
 	virtual FArchive* GetStreamingArchive() override;
