Use condensed version of json when sending user list

diff --git a/Engine/Source/Runtime/NetworkReplayStreaming/HttpNetworkReplayStreaming/Private/HttpNetworkReplayStreaming.cpp b/Engine/Source/Runtime/NetworkReplayStreaming/HttpNetworkReplayStreaming/Private/HttpNetworkReplayStreaming.cpp
--- a/Engine/Source/Runtime/NetworkReplayStreaming/HttpNetworkReplayStreaming/Private/HttpNetworkReplayStreaming.cpp
+++ b/Engine/Source/Runtime/NetworkReplayStreaming/HttpNetworkReplayStreaming/Private/HttpNetworkReplayStreaming.cpp
@@ -251,7 +251,7 @@ void FHttpNetworkReplayStreamer::StartStreaming( const FString& CustomName, cons
 			FNetworkReplayUserList UserList;
 
 			UserList.Users = UserNames;
-			HttpRequest->SetContentAsString( UserList.ToJson() );
+			HttpRequest->SetContentAsString( UserList.ToJson(false) );
 			HttpRequest->SetHeader( TEXT( "Content-Type" ), TEXT( "application/json" ) );
 		}
 
