Merge:  Assume failed downloads with live streams are ok + clamp scrubbing to last checkpoint if no stream is available after it

diff --git a/Engine/Source/Runtime/NetworkReplayStreaming/HttpNetworkReplayStreaming/Private/HttpNetworkReplayStreaming.cpp b/Engine/Source/Runtime/NetworkReplayStreaming/HttpNetworkReplayStreaming/Private/HttpNetworkReplayStreaming.cpp
--- a/Engine/Source/Runtime/NetworkReplayStreaming/HttpNetworkReplayStreaming/Private/HttpNetworkReplayStreaming.cpp
+++ b/Engine/Source/Runtime/NetworkReplayStreaming/HttpNetworkReplayStreaming/Private/HttpNetworkReplayStreaming.cpp
@@ -1067,6 +1067,12 @@ void FHttpNetworkReplayStreamer::HttpDownloadFinished( FHttpRequestPtr HttpReque
 	}
 	else
 	{
+		if ( bStreamIsLive )
+		{
+			// Assume this isn't really an error
+			return;
+		}
+
 		UE_LOG( LogHttpReplay, Error, TEXT( "FHttpNetworkReplayStreamer::HttpDownloadFinished. FAILED." ) );
 		StreamArchive.Buffer.Empty();
 		SetLastError( ENetworkReplayError::ServiceUnavailable );
@@ -1102,6 +1108,9 @@ void FHttpNetworkReplayStreamer::HttpDownloadCheckpointFinished( FHttpRequestPtr
 		StreamArchive.Pos				= 0;
 		StreamArchive.bAtEndOfReplay	= false;
 
+		// Reset any time we were waiting on in the past
+		HighPriorityEndTime	= 0;
+
 		// Reset our stream range
 		StreamTimeRangeStart	= 0;
 		StreamTimeRangeEnd		= 0;
@@ -1109,6 +1118,16 @@ void FHttpNetworkReplayStreamer::HttpDownloadCheckpointFinished( FHttpRequestPtr
 		// Set the next chunk to be right after this checkpoint (which was stored in the metadata)
 		StreamChunkIndex = FCString::Atoi( *CheckpointList.Checkpoints[ DownloadCheckpointIndex ].Metadata );
 
+		// If we want to fast forward past the end of a stream, clamp to the checkpoint
+		if ( LastGotoTimeInMS >= 0 && StreamChunkIndex >= NumTotalStreamChunks )
+		{
+			UE_LOG( LogHttpReplay, Warning, TEXT( "FHttpNetworkReplayStreamer::HttpDownloadCheckpointFinished. Clamped to checkpoint: %i" ), LastGotoTimeInMS );
+
+			StreamTimeRangeStart	= CheckpointList.Checkpoints[DownloadCheckpointIndex].Time1;
+			StreamTimeRangeEnd		= CheckpointList.Checkpoints[DownloadCheckpointIndex].Time1;
+			LastGotoTimeInMS		= -1;
+		}
+
 		if ( LastGotoTimeInMS >= 0 )
 		{
 			// If we are fine scrubbing, make sure to wait on the part of the stream that is needed to do this in one frame
