Merge: More HTTP replay streamer debug info/logging

diff --git a/Engine/Source/Runtime/NetworkReplayStreaming/HttpNetworkReplayStreaming/Private/HttpNetworkReplayStreaming.cpp b/Engine/Source/Runtime/NetworkReplayStreaming/HttpNetworkReplayStreaming/Private/HttpNetworkReplayStreaming.cpp
--- a/Engine/Source/Runtime/NetworkReplayStreaming/HttpNetworkReplayStreaming/Private/HttpNetworkReplayStreaming.cpp
+++ b/Engine/Source/Runtime/NetworkReplayStreaming/HttpNetworkReplayStreaming/Private/HttpNetworkReplayStreaming.cpp
@@ -205,7 +205,11 @@ void FHttpNetworkReplayStreamer::StartStreaming( const FString& CustomName, cons
 		}
 
 		// Notify the http server that we want to start downloading a replay
-		HttpRequest->SetURL( FString::Printf( TEXT( "%sstartdownloading?Session=%s&User=%s" ), *ServerURL, *SessionName, *UserName ) );
+		const FString URL = FString::Printf( TEXT( "%sstartdownloading?Session=%s&User=%s" ), *ServerURL, *SessionName, *UserName );
+
+		UE_LOG( LogHttpReplay, Verbose, TEXT( "FHttpNetworkReplayStreamer::StartStreaming. URL: %s" ), *URL );
+
+		HttpRequest->SetURL( URL );
 		HttpRequest->SetVerb( TEXT( "POST" ) );
 
 		HttpRequest->OnProcessRequestComplete().BindRaw( this, &FHttpNetworkReplayStreamer::HttpStartDownloadingFinished );
@@ -259,9 +263,9 @@ void FHttpNetworkReplayStreamer::StartStreaming( const FString& CustomName, cons
 	}
 }
 
-void FHttpNetworkReplayStreamer::AddRequestToQueue( const EQueuedHttpRequestType Type, TSharedPtr< class IHttpRequest >	Request )
+void FHttpNetworkReplayStreamer::AddRequestToQueue( const EQueuedHttpRequestType::Type Type, TSharedPtr< class IHttpRequest >	Request )
 {
-	UE_LOG( LogHttpReplay, Verbose, TEXT( "FHttpNetworkReplayStreamer::AddRequestToQueue. Type: %i" ), (int)Type );
+	UE_LOG( LogHttpReplay, Verbose, TEXT( "FHttpNetworkReplayStreamer::AddRequestToQueue. Type: %s" ), EQueuedHttpRequestType::ToString( Type ) );
 
 	QueuedHttpRequests.Enqueue( TSharedPtr< FQueuedHttpRequest >( new FQueuedHttpRequest( Type, Request ) ) );
 }
@@ -603,7 +607,11 @@ void FHttpNetworkReplayStreamer::ConditionallyDownloadNextChunk()
 	TSharedRef<class IHttpRequest> HttpRequest = FHttpModule::Get().CreateRequest();
 
 	// Download the next stream chunk
-	HttpRequest->SetURL( FString::Printf( TEXT( "%sdownload?Session=%s&Filename=stream.%i" ), *ServerURL, *SessionName, StreamChunkIndex ) );
+	const FString URL = FString::Printf( TEXT( "%sdownload?Session=%s&Filename=stream.%i" ), *ServerURL, *SessionName, StreamChunkIndex );
+
+	UE_LOG( LogHttpReplay, Verbose, TEXT( "FHttpNetworkReplayStreamer::ConditionallyDownloadNextChunk. URL: %s" ), *URL );
+
+	HttpRequest->SetURL( URL );
 	HttpRequest->SetVerb( TEXT( "GET" ) );
 
 	HttpRequest->OnProcessRequestComplete().BindRaw( this, &FHttpNetworkReplayStreamer::HttpDownloadFinished );
@@ -853,7 +861,7 @@ void FHttpNetworkReplayStreamer::ConditionallyEnumerateCheckpoints()
 	}
 };
 
-void FHttpNetworkReplayStreamer::RequestFinished( EStreamerState ExpectedStreamerState, EQueuedHttpRequestType ExpectedType, FHttpRequestPtr HttpRequest )
+void FHttpNetworkReplayStreamer::RequestFinished( EStreamerState ExpectedStreamerState, EQueuedHttpRequestType::Type ExpectedType, FHttpRequestPtr HttpRequest )
 {
 	check( StreamerState == ExpectedStreamerState );
 	check( InFlightHttpRequest.IsValid() );
@@ -988,7 +996,7 @@ void FHttpNetworkReplayStreamer::HttpStartDownloadingFinished( FHttpRequestPtr H
 	}
 	else
 	{
-		UE_LOG( LogHttpReplay, Error, TEXT( "FHttpNetworkReplayStreamer::HttpStartDownloadingFinished. FAILED" ) );
+		UE_LOG( LogHttpReplay, Error, TEXT( "FHttpNetworkReplayStreamer::HttpStartDownloadingFinished. FAILED. Code: %i" ), HttpResponse->GetResponseCode() );
 
 		StartStreamingDelegate.ExecuteIfBound( false, true );
 
@@ -1272,7 +1280,7 @@ bool FHttpNetworkReplayStreamer::ProcessNextHttpRequest()
 
 	if ( QueuedHttpRequests.Dequeue( QueuedRequest ) )
 	{
-		UE_LOG( LogHttpReplay, Verbose, TEXT( "FHttpNetworkReplayStreamer::ProcessNextHttpRequest. Dequeue Type: %i" ), (int)QueuedRequest->Type );
+		UE_LOG( LogHttpReplay, Verbose, TEXT( "FHttpNetworkReplayStreamer::ProcessNextHttpRequest. Dequeue Type: %s" ), EQueuedHttpRequestType::ToString( QueuedRequest->Type ) );
 
 		check( !InFlightHttpRequest.IsValid() );
 
diff --git a/Engine/Source/Runtime/NetworkReplayStreaming/HttpNetworkReplayStreaming/Private/HttpNetworkReplayStreaming.h b/Engine/Source/Runtime/NetworkReplayStreaming/HttpNetworkReplayStreaming/Private/HttpNetworkReplayStreaming.h
--- a/Engine/Source/Runtime/NetworkReplayStreaming/HttpNetworkReplayStreaming/Private/HttpNetworkReplayStreaming.h
+++ b/Engine/Source/Runtime/NetworkReplayStreaming/HttpNetworkReplayStreaming/Private/HttpNetworkReplayStreaming.h
@@ -62,28 +62,66 @@ public:
 	bool			bAtEndOfReplay;
 };
 
-enum class EQueuedHttpRequestType
+namespace EQueuedHttpRequestType
 {
-	StartUploading,				// We have made a request to start uploading a replay
-	UploadingHeader,			// We are uploading the replay header
-	UploadingStream,			// We are in the process of uploading the replay stream
-	StopUploading,				// We have made the request to stop uploading a live replay stream
-	StartDownloading,			// We have made the request to start downloading a replay stream
-	DownloadingHeader,			// We are downloading the replay header
-	DownloadingStream,			// We are in the process of downloading the replay stream
-	RefreshingViewer,			// We are refreshing the server to let it know we're still viewing
-	EnumeratingSessions,		// We are in the process of downloading the available sessions
-	EnumeratingCheckpoints,		// We are in the process of downloading the available checkpoints
-	UploadingCheckpoint,		// We are uploading a checkpoint
-	DownloadingCheckpoint		// We are downloading a checkpoint
+	enum Type
+	{
+		StartUploading,				// We have made a request to start uploading a replay
+		UploadingHeader,			// We are uploading the replay header
+		UploadingStream,			// We are in the process of uploading the replay stream
+		StopUploading,				// We have made the request to stop uploading a live replay stream
+		StartDownloading,			// We have made the request to start downloading a replay stream
+		DownloadingHeader,			// We are downloading the replay header
+		DownloadingStream,			// We are in the process of downloading the replay stream
+		RefreshingViewer,			// We are refreshing the server to let it know we're still viewing
+		EnumeratingSessions,		// We are in the process of downloading the available sessions
+		EnumeratingCheckpoints,		// We are in the process of downloading the available checkpoints
+		UploadingCheckpoint,		// We are uploading a checkpoint
+		DownloadingCheckpoint		// We are downloading a checkpoint
+	};
+
+	inline const TCHAR* ToString( EQueuedHttpRequestType::Type Type )
+	{
+		switch ( Type )
+		{
+			case StartUploading:
+				return TEXT( "StartUploading" );
+			case UploadingHeader:
+				return TEXT( "UploadingHeader" );
+			case UploadingStream:
+				return TEXT( "UploadingStream" );
+			case StopUploading:
+				return TEXT( "StopUploading" );
+			case StartDownloading:
+				return TEXT( "StartDownloading" );
+			case DownloadingHeader:
+				return TEXT( "DownloadingHeader" );
+			case DownloadingStream:
+				return TEXT( "DownloadingStream" );
+			case RefreshingViewer:
+				return TEXT( "RefreshingViewer" );
+			case EnumeratingSessions:
+				return TEXT( "EnumeratingSessions" );
+			case EnumeratingCheckpoints:
+				return TEXT( "EnumeratingCheckpoints" );
+			case UploadingCheckpoint:
+				return TEXT( "UploadingCheckpoint" );
+			case DownloadingCheckpoint:
+				return TEXT( "DownloadingCheckpoint" );
+		}
+
+		return TEXT( "Unknown EQueuedHttpRequestType type." );
+	}
 };
 
 class FQueuedHttpRequest
 {
 public:
-	FQueuedHttpRequest( const EQueuedHttpRequestType InType, TSharedPtr< class IHttpRequest > InRequest ) : Type( InType ), Request( InRequest ) { }
+	FQueuedHttpRequest( const EQueuedHttpRequestType::Type InType, TSharedPtr< class IHttpRequest > InRequest ) : Type( InType ), Request( InRequest )
+	{
+	}
 
-	EQueuedHttpRequestType				Type;
+	EQueuedHttpRequestType::Type		Type;
 	TSharedPtr< class IHttpRequest >	Request;
 };
 
@@ -129,7 +167,7 @@ public:
 	void ConditionallyRefreshViewer();
 	void SetLastError( const ENetworkReplayError::Type InLastError );
 	void FlushCheckpointInternal( uint32 TimeInMS );
-	void AddRequestToQueue( const EQueuedHttpRequestType Type, TSharedPtr< class IHttpRequest >	Request );
+	void AddRequestToQueue( const EQueuedHttpRequestType::Type Type, TSharedPtr< class IHttpRequest >	Request );
 	void EnumerateCheckpoints();
 	void ConditionallyEnumerateCheckpoints();
 
@@ -142,7 +180,7 @@ public:
 	};
 
 	/** Delegates */
-	void RequestFinished( EStreamerState ExpectedStreamerState, EQueuedHttpRequestType ExpectedType, FHttpRequestPtr HttpRequest );
+	void RequestFinished( EStreamerState ExpectedStreamerState, EQueuedHttpRequestType::Type ExpectedType, FHttpRequestPtr HttpRequest );
 
 	void HttpStartDownloadingFinished( FHttpRequestPtr HttpRequest, FHttpResponsePtr HttpResponse, bool bSucceeded );
 	void HttpDownloadHeaderFinished( FHttpRequestPtr HttpRequest, FHttpResponsePtr HttpResponse, bool bSucceeded );
