Clamp demo time when scrubbing

diff --git a/Engine/Source/Runtime/NetworkReplayStreaming/HttpNetworkReplayStreaming/Private/HttpNetworkReplayStreaming.cpp b/Engine/Source/Runtime/NetworkReplayStreaming/HttpNetworkReplayStreaming/Private/HttpNetworkReplayStreaming.cpp
--- a/Engine/Source/Runtime/NetworkReplayStreaming/HttpNetworkReplayStreaming/Private/HttpNetworkReplayStreaming.cpp
+++ b/Engine/Source/Runtime/NetworkReplayStreaming/HttpNetworkReplayStreaming/Private/HttpNetworkReplayStreaming.cpp
@@ -496,7 +496,7 @@ void FHttpNetworkReplayStreamer::GotoTimeInMS( const uint32 TimeInMS, const FOnC
 
 	int32 CheckpointIndex = -1;
 
-	LastGotoTimeInMS = TimeInMS;
+	LastGotoTimeInMS = FMath::Min( TimeInMS, TotalDemoTimeInMS );
 
 	if ( CheckpointList.Checkpoints.Num() > 0 && TimeInMS >= CheckpointList.Checkpoints[ CheckpointList.Checkpoints.Num() - 1 ].Time1 )
 	{
