Fix crash that could occur while scrubbing in replays

diff --git a/Engine/Source/Runtime/NetworkReplayStreaming/HttpNetworkReplayStreaming/Private/HttpNetworkReplayStreaming.cpp b/Engine/Source/Runtime/NetworkReplayStreaming/HttpNetworkReplayStreaming/Private/HttpNetworkReplayStreaming.cpp
--- a/Engine/Source/Runtime/NetworkReplayStreaming/HttpNetworkReplayStreaming/Private/HttpNetworkReplayStreaming.cpp
+++ b/Engine/Source/Runtime/NetworkReplayStreaming/HttpNetworkReplayStreaming/Private/HttpNetworkReplayStreaming.cpp
@@ -481,6 +481,12 @@ void FHttpNetworkReplayStreamer::GotoCheckpointIndex( const int32 CheckpointInde
 
 void FHttpNetworkReplayStreamer::GotoTimeInMS( const uint32 TimeInMS, const FOnCheckpointReadyDelegate& Delegate )
 {
+	if ( IsHttpRequestInFlight() || HasPendingHttpRequests() )
+	{
+		UE_LOG( LogHttpReplay, Log, TEXT( "FHttpNetworkReplayStreamer::GotoTimeInMS. Busy processing pending requests." ) );
+		return;
+	}
+
 	if ( GotoCheckpointDelegate.IsBound() )
 	{
 		// If we're currently going to a checkpoint now, ignore this request
