#Orion: Allow override of meta tag when viewing replays

--------
Integrated using branch UE4-To-UE4-Orion (reversed) of change#2576985 by [REDACTED_PERSON] on 2015/06/04 10:47:22.

diff --git a/Engine/Source/Runtime/NetworkReplayStreaming/HttpNetworkReplayStreaming/Private/HttpNetworkReplayStreaming.cpp b/Engine/Source/Runtime/NetworkReplayStreaming/HttpNetworkReplayStreaming/Private/HttpNetworkReplayStreaming.cpp
--- a/Engine/Source/Runtime/NetworkReplayStreaming/HttpNetworkReplayStreaming/Private/HttpNetworkReplayStreaming.cpp
+++ b/Engine/Source/Runtime/NetworkReplayStreaming/HttpNetworkReplayStreaming/Private/HttpNetworkReplayStreaming.cpp
@@ -4,6 +4,8 @@
 
 DEFINE_LOG_CATEGORY_STATIC( LogHttpReplay, Log, All );
 
+static TAutoConsoleVariable<FString> CVarMetaFilterOverride( TEXT( "httpReplay.MetaFilterOverride" ), TEXT( "" ), TEXT( "" ) );
+
 class FNetworkReplayListItem : public FOnlineJsonSerializable
 {
 public:
@@ -792,12 +794,15 @@ void FHttpNetworkReplayStreamer::EnumerateStreams( const FNetworkReplayVersion&
 	// Build base URL
 	FString URL = FString::Printf( TEXT( "%senumsessions?App=%s&Version=%u&CL=%u" ), *ServerURL, *InReplayVersion.AppString, InReplayVersion.NetworkVersion, InReplayVersion.Changelist );
 
-	// Add optional parameters
-	if ( !MetaString.IsEmpty() )
+	const FString MetaStringToUse = !CVarMetaFilterOverride.GetValueOnGameThread().IsEmpty() ? CVarMetaFilterOverride.GetValueOnGameThread() : MetaString;
+
+	// Add optional Meta parameter (filter replays by meta tag)
+	if ( !MetaStringToUse.IsEmpty() )
 	{
-		URL += FString::Printf( TEXT( "&Meta=%s" ), *MetaString );
+		URL += FString::Printf( TEXT( "&Meta=%s" ), *MetaStringToUse );
 	}
 
+	// Add optional User parameter (filter replays by a user that was in the replay)
 	if ( !UserString.IsEmpty() )
 	{
 		URL += FString::Printf( TEXT( "&User=%s" ), *UserString );
