Merge: Allow HttpNetworkReplayStreaming to be derived from

diff --git a/Engine/Source/Runtime/NetworkReplayStreaming/HttpNetworkReplayStreaming/HttpNetworkReplayStreaming.Build.cs b/Engine/Source/Runtime/NetworkReplayStreaming/HttpNetworkReplayStreaming/HttpNetworkReplayStreaming.Build.cs
--- a/Engine/Source/Runtime/NetworkReplayStreaming/HttpNetworkReplayStreaming/HttpNetworkReplayStreaming.Build.cs
+++ b/Engine/Source/Runtime/NetworkReplayStreaming/HttpNetworkReplayStreaming/HttpNetworkReplayStreaming.Build.cs
@@ -6,8 +6,6 @@ namespace UnrealBuildTool.Rules
     {
         public HttpNetworkReplayStreaming(TargetInfo Target)
         {
-			PrivateIncludePaths.Add( "Runtime/NetworkReplayStreaming/HttpNetworkReplayStreaming/Private" );
-
 			PrivateIncludePathModuleNames.Add( "OnlineSubsystem" );
 
 			PrivateDependencyModuleNames.AddRange(
diff --git a/Engine/Source/Runtime/NetworkReplayStreaming/HttpNetworkReplayStreaming/Private/HttpNetworkReplayStreaming.h b/Engine/Source/Runtime/NetworkReplayStreaming/HttpNetworkReplayStreaming/Private/HttpNetworkReplayStreaming.h
deleted file mode 100644
diff --git a/Engine/Source/Runtime/NetworkReplayStreaming/HttpNetworkReplayStreaming/Public/HttpNetworkReplayStreaming.h b/Engine/Source/Runtime/NetworkReplayStreaming/HttpNetworkReplayStreaming/Public/HttpNetworkReplayStreaming.h
new file mode 100644
--- /dev/null
+++ b/Engine/Source/Runtime/NetworkReplayStreaming/HttpNetworkReplayStreaming/Public/HttpNetworkReplayStreaming.h
@@ -0,0 +1,252 @@
+// Copyright 1998-2015 Epic Games, Inc. All Rights Reserved.
+#pragma once
+
+#include "NetworkReplayStreaming.h"
+#include "Http.h"
+#include "Runtime/Engine/Public/Tickable.h"
+#include "OnlineJsonSerializer.h"
+
+class FCheckpointListItem : public FOnlineJsonSerializable
+{
+public:
+	FCheckpointListItem() {}
+	virtual ~FCheckpointListItem() {}
+
+	FString		ID;
+	FString		Group;
+	FString		Metadata;
+	uint32		Time1;
+	uint32		Time2;
+
+	// FOnlineJsonSerializable
+	BEGIN_ONLINE_JSON_SERIALIZER
+		ONLINE_JSON_SERIALIZE( "id",			ID );
+		ONLINE_JSON_SERIALIZE( "group",			Group );
+		ONLINE_JSON_SERIALIZE( "meta",			Metadata );
+		ONLINE_JSON_SERIALIZE( "time1",			Time1 );
+		ONLINE_JSON_SERIALIZE( "time2",			Time2 );
+	END_ONLINE_JSON_SERIALIZER
+};
+
+class FCheckpointList : public FOnlineJsonSerializable
+{
+public:
+	FCheckpointList()
+	{}
+	virtual ~FCheckpointList() {}
+
+	TArray< FCheckpointListItem > Checkpoints;
+
+	// FOnlineJsonSerializable
+	BEGIN_ONLINE_JSON_SERIALIZER
+		ONLINE_JSON_SERIALIZE_ARRAY_SERIALIZABLE( "events", Checkpoints, FCheckpointListItem );
+	END_ONLINE_JSON_SERIALIZER
+};
+
+/**
+ * Archive used to buffer stream over http
+ */
+class FHttpStreamFArchive : public FArchive
+{
+public:
+	FHttpStreamFArchive() : Pos( 0 ), bAtEndOfReplay( false ) {}
+
+	virtual void	Serialize( void* V, int64 Length );
+	virtual int64	Tell();
+	virtual int64	TotalSize();
+	virtual void	Seek( int64 InPos );
+	virtual bool	AtEnd();
+
+	TArray< uint8 >	Buffer;
+	int32			Pos;
+	bool			bAtEndOfReplay;
+};
+
+namespace EQueuedHttpRequestType
+{
+	enum Type
+	{
+		StartUploading,				// We have made a request to start uploading a replay
+		UploadingHeader,			// We are uploading the replay header
+		UploadingStream,			// We are in the process of uploading the replay stream
+		StopUploading,				// We have made the request to stop uploading a live replay stream
+		StartDownloading,			// We have made the request to start downloading a replay stream
+		DownloadingHeader,			// We are downloading the replay header
+		DownloadingStream,			// We are in the process of downloading the replay stream
+		RefreshingViewer,			// We are refreshing the server to let it know we're still viewing
+		EnumeratingSessions,		// We are in the process of downloading the available sessions
+		EnumeratingCheckpoints,		// We are in the process of downloading the available checkpoints
+		UploadingCheckpoint,		// We are uploading a checkpoint
+		DownloadingCheckpoint		// We are downloading a checkpoint
+	};
+
+	inline const TCHAR* ToString( EQueuedHttpRequestType::Type Type )
+	{
+		switch ( Type )
+		{
+			case StartUploading:
+				return TEXT( "StartUploading" );
+			case UploadingHeader:
+				return TEXT( "UploadingHeader" );
+			case UploadingStream:
+				return TEXT( "UploadingStream" );
+			case StopUploading:
+				return TEXT( "StopUploading" );
+			case StartDownloading:
+				return TEXT( "StartDownloading" );
+			case DownloadingHeader:
+				return TEXT( "DownloadingHeader" );
+			case DownloadingStream:
+				return TEXT( "DownloadingStream" );
+			case RefreshingViewer:
+				return TEXT( "RefreshingViewer" );
+			case EnumeratingSessions:
+				return TEXT( "EnumeratingSessions" );
+			case EnumeratingCheckpoints:
+				return TEXT( "EnumeratingCheckpoints" );
+			case UploadingCheckpoint:
+				return TEXT( "UploadingCheckpoint" );
+			case DownloadingCheckpoint:
+				return TEXT( "DownloadingCheckpoint" );
+		}
+
+		return TEXT( "Unknown EQueuedHttpRequestType type." );
+	}
+};
+
+class FQueuedHttpRequest
+{
+public:
+	FQueuedHttpRequest( const EQueuedHttpRequestType::Type InType, TSharedPtr< class IHttpRequest > InRequest ) : Type( InType ), Request( InRequest )
+	{
+	}
+
+	EQueuedHttpRequestType::Type		Type;
+	TSharedPtr< class IHttpRequest >	Request;
+};
+
+/**
+ * Http network replay streaming manager
+ */
+class HTTPNETWORKREPLAYSTREAMING_API FHttpNetworkReplayStreamer : public INetworkReplayStreamer
+{
+public:
+	FHttpNetworkReplayStreamer();
+
+	/** INetworkReplayStreamer implementation */
+	virtual void		StartStreaming( const FString& CustomName, const FString& FriendlyName, const TArray< FString >& UserNames, bool bRecord, const FNetworkReplayVersion& ReplayVersion, const FOnStreamReadyDelegate& Delegate ) override;
+	virtual void		StopStreaming() override;
+	virtual FArchive*	GetHeaderArchive() override;
+	virtual FArchive*	GetStreamingArchive() override;
+	virtual FArchive*	GetCheckpointArchive() override;
+	virtual void		FlushCheckpoint( const uint32 TimeInMS ) override;
+	virtual void		GotoCheckpointIndex( const int32 CheckpointIndex, const FOnCheckpointReadyDelegate& Delegate ) override;
+	virtual void		GotoTimeInMS( const uint32 TimeInMS, const FOnCheckpointReadyDelegate& Delegate ) override;
+	virtual FArchive*	GetMetadataArchive() override;
+	virtual void		UpdateTotalDemoTime( uint32 TimeInMS ) override;
+	virtual uint32		GetTotalDemoTime() const override { return TotalDemoTimeInMS; }
+	virtual bool		IsDataAvailable() const override;
+	virtual void		SetHighPriorityTimeRange( const uint32 StartTimeInMS, const uint32 EndTimeInMS ) override;
+	virtual bool		IsDataAvailableForTimeRange( const uint32 StartTimeInMS, const uint32 EndTimeInMS ) override;
+	virtual bool		IsLoadingCheckpoint() const override;
+	virtual bool		IsLive() const override;
+	virtual void		DeleteFinishedStream( const FString& StreamName, const FOnDeleteFinishedStreamComplete& Delegate ) const override;
+	virtual void		EnumerateStreams( const FNetworkReplayVersion& ReplayVersion, const FString& UserString, const FString& MetaString, const FOnEnumerateStreamsComplete& Delegate ) override;
+	virtual void		EnumerateRecentStreams( const FNetworkReplayVersion& ReplayVersion, const FString& RecentViewer, const FOnEnumerateStreamsComplete& Delegate ) override;
+
+	virtual ENetworkReplayError::Type GetLastError() const override;
+
+	/** FHttpNetworkReplayStreamer */
+	void UploadHeader();
+	void FlushStream();
+	void ConditionallyFlushStream();
+	void StopUploading();
+	void DownloadHeader();
+	void ConditionallyDownloadNextChunk();
+	void RefreshViewer( const bool bFinal );
+	void ConditionallyRefreshViewer();
+	void SetLastError( const ENetworkReplayError::Type InLastError );
+	void FlushCheckpointInternal( uint32 TimeInMS );
+	void AddRequestToQueue( const EQueuedHttpRequestType::Type Type, TSharedPtr< class IHttpRequest >	Request );
+	void EnumerateCheckpoints();
+	void ConditionallyEnumerateCheckpoints();
+
+	/** EStreamerState - Overall state of the streamer */
+	enum class EStreamerState
+	{
+		Idle,						// The streamer is idle. Either we haven't started streaming yet, or we are done
+		StreamingUp,				// We are in the process of streaming a replay to the http server
+		StreamingDown				// We are in the process of streaming a replay from the http server
+	};
+
+	/** Delegates */
+	void RequestFinished( EStreamerState ExpectedStreamerState, EQueuedHttpRequestType::Type ExpectedType, FHttpRequestPtr HttpRequest );
+
+	void HttpStartDownloadingFinished( FHttpRequestPtr HttpRequest, FHttpResponsePtr HttpResponse, bool bSucceeded );
+	void HttpDownloadHeaderFinished( FHttpRequestPtr HttpRequest, FHttpResponsePtr HttpResponse, bool bSucceeded );
+	void HttpDownloadFinished( FHttpRequestPtr HttpRequest, FHttpResponsePtr HttpResponse, bool bSucceeded );
+	void HttpDownloadCheckpointFinished( FHttpRequestPtr HttpRequest, FHttpResponsePtr HttpResponse, bool bSucceeded );
+	void HttpRefreshViewerFinished( FHttpRequestPtr HttpRequest, FHttpResponsePtr HttpResponse, bool bSucceeded );
+	void HttpStartUploadingFinished( FHttpRequestPtr HttpRequest, FHttpResponsePtr HttpResponse, bool bSucceeded );
+	void HttpStopUploadingFinished( FHttpRequestPtr HttpRequest, FHttpResponsePtr HttpResponse, bool bSucceeded );
+	void HttpHeaderUploadFinished( FHttpRequestPtr HttpRequest, FHttpResponsePtr HttpResponse, bool bSucceeded );
+	void HttpUploadStreamFinished( FHttpRequestPtr HttpRequest, FHttpResponsePtr HttpResponse, bool bSucceeded );
+	void HttpUploadCheckpointFinished( FHttpRequestPtr HttpRequest, FHttpResponsePtr HttpResponse, bool bSucceeded );
+	void HttpEnumerateSessionsFinished( FHttpRequestPtr HttpRequest, FHttpResponsePtr HttpResponse, bool bSucceeded );
+	void HttpEnumerateCheckpointsFinished( FHttpRequestPtr HttpRequest, FHttpResponsePtr HttpResponse, bool bSucceeded );
+
+	bool ProcessNextHttpRequest();
+	void Tick( const float DeltaTime );
+	bool IsHttpRequestInFlight() const;		// True there is an http request currently in flight
+	bool HasPendingHttpRequests() const;	// True if there is an http request in flight, or there are more to process
+	bool IsStreaming() const;				// True if we are streaming a replay up or down
+
+	FHttpStreamFArchive		HeaderArchive;			// Archive used to buffer the header stream
+	FHttpStreamFArchive		StreamArchive;			// Archive used to buffer the data stream
+	FHttpStreamFArchive		CheckpointArchive;		// Archive used to buffer checkpoint data
+	FString					SessionName;			// Name of the session on the http replay server
+	FNetworkReplayVersion	ReplayVersion;			// Version of the session
+	FString					ServerURL;				// The address of the server
+	int32					StreamChunkIndex;		// Used as a counter to increment the stream.x extension count
+	double					LastChunkTime;			// The last time we uploaded/downloaded a chunk
+	double					LastRefreshViewerTime;	// The last time we refreshed ourselves as an active viewer
+	double					LastRefreshCheckpointTime;
+	EStreamerState			StreamerState;			// Overall state of the streamer
+	bool					bStopStreamingCalled;
+	bool					bNeedToUploadHeader;	// We're waiting on session name so we can upload header
+	bool					bStreamIsLive;			// If true, we are viewing a live stream
+	int32					NumTotalStreamChunks;
+	uint32					TotalDemoTimeInMS;
+	uint32					LastTotalDemoTimeInMS;
+	uint32					StreamTimeRangeStart;
+	uint32					StreamTimeRangeEnd;
+	FString					ViewerName;
+	uint32					HighPriorityEndTime;
+
+	ENetworkReplayError::Type		StreamerLastError;
+
+	FOnStreamReadyDelegate			StartStreamingDelegate;		// Delegate passed in to StartStreaming
+	FOnEnumerateStreamsComplete		EnumerateStreamsDelegate;
+	FOnCheckpointReadyDelegate		GotoCheckpointDelegate;
+	int32							DownloadCheckpointIndex;
+	int64							LastGotoTimeInMS;
+
+	FCheckpointList					CheckpointList;
+
+	TQueue< TSharedPtr< FQueuedHttpRequest > >	QueuedHttpRequests;
+	TSharedPtr< FQueuedHttpRequest >			InFlightHttpRequest;
+};
+
+class HTTPNETWORKREPLAYSTREAMING_API FHttpNetworkReplayStreamingFactory : public INetworkReplayStreamingFactory, public FTickableGameObject
+{
+public:
+	/** INetworkReplayStreamingFactory */
+	virtual TSharedPtr< INetworkReplayStreamer > CreateReplayStreamer() override;
+
+	/** FTickableGameObject */
+	virtual void Tick( float DeltaTime ) override;
+	virtual bool IsTickable() const override;
+	virtual TStatId GetStatId() const override;
+
+	TArray< TSharedPtr< FHttpNetworkReplayStreamer > > HttpStreamers;
+};
