Merge:  Sort replay checkpoints by time

diff --git a/Engine/Source/Runtime/NetworkReplayStreaming/HttpNetworkReplayStreaming/Private/HttpNetworkReplayStreaming.cpp b/Engine/Source/Runtime/NetworkReplayStreaming/HttpNetworkReplayStreaming/Private/HttpNetworkReplayStreaming.cpp
--- a/Engine/Source/Runtime/NetworkReplayStreaming/HttpNetworkReplayStreaming/Private/HttpNetworkReplayStreaming.cpp
+++ b/Engine/Source/Runtime/NetworkReplayStreaming/HttpNetworkReplayStreaming/Private/HttpNetworkReplayStreaming.cpp
@@ -1262,6 +1262,17 @@ void FHttpNetworkReplayStreamer::HttpEnumerateCheckpointsFinished( FHttpRequestP
 			return;
 		}
 
+		// Sort checkpoints by time
+		struct FCompareCheckpointTime
+		{
+			FORCEINLINE bool operator()( const FCheckpointListItem & A, const FCheckpointListItem & B ) const
+			{
+				return A.Time1 < B.Time1;
+			}
+		};
+
+		Sort( CheckpointList.Checkpoints.GetData(), CheckpointList.Checkpoints.Num(), FCompareCheckpointTime() );
+
 		StartStreamingDelegate.ExecuteIfBound( true, false );
 	}
 	else
