Merge: Fix an issue when the file download command returns success with no status update (assume still live)

diff --git a/Engine/Source/Runtime/NetworkReplayStreaming/HttpNetworkReplayStreaming/Private/HttpNetworkReplayStreaming.cpp b/Engine/Source/Runtime/NetworkReplayStreaming/HttpNetworkReplayStreaming/Private/HttpNetworkReplayStreaming.cpp
--- a/Engine/Source/Runtime/NetworkReplayStreaming/HttpNetworkReplayStreaming/Private/HttpNetworkReplayStreaming.cpp
+++ b/Engine/Source/Runtime/NetworkReplayStreaming/HttpNetworkReplayStreaming/Private/HttpNetworkReplayStreaming.cpp
@@ -1054,6 +1054,13 @@ void FHttpNetworkReplayStreamer::HttpDownloadFinished( FHttpRequestPtr HttpReque
 
 	if ( bSucceeded && HttpResponse->GetResponseCode() == EHttpResponseCodes::Ok )
 	{
+		if ( HttpResponse->GetHeader( TEXT( "NumChunks" ) ) == TEXT( "" ) )
+		{
+			// Assume this is an implcit status update
+			UE_LOG( LogHttpReplay, Log, TEXT( "FHttpNetworkReplayStreamer::HttpDownloadFinished. NO HEADER FIELDS. Live: %i, Progress: %i / %i, Start: %i, End: %i, DemoTime: %2.2f" ), ( int )bStreamIsLive, StreamChunkIndex, NumTotalStreamChunks, ( int )StreamTimeRangeStart, ( int )StreamTimeRangeEnd, ( float )TotalDemoTimeInMS / 1000 );
+			return;
+		}
+
 		NumTotalStreamChunks = FCString::Atoi( *HttpResponse->GetHeader( TEXT( "NumChunks" ) ) );
 		TotalDemoTimeInMS = FCString::Atoi( *HttpResponse->GetHeader( TEXT( "Time" ) ) );
 		bStreamIsLive = HttpResponse->GetHeader( TEXT( "State" ) ) == TEXT( "Live" );
@@ -1088,13 +1095,16 @@ void FHttpNetworkReplayStreamer::HttpDownloadFinished( FHttpRequestPtr HttpReque
 	{
 		if ( bStreamIsLive )
 		{
-			// Assume this isn't really an error
-			return;
-		}
+			bStreamIsLive = false;
 
-		UE_LOG( LogHttpReplay, Error, TEXT( "FHttpNetworkReplayStreamer::HttpDownloadFinished. FAILED." ) );
-		StreamArchive.Buffer.Empty();
-		SetLastError( ENetworkReplayError::ServiceUnavailable );
+			UE_LOG( LogHttpReplay, Warning, TEXT( "FHttpNetworkReplayStreamer::HttpDownloadFinished. Failed live, turning off live flag. Live: %i, Progress: %i / %i, Start: %i, End: %i, DemoTime: %2.2f" ), ( int )bStreamIsLive, StreamChunkIndex, NumTotalStreamChunks, ( int )StreamTimeRangeStart, ( int )StreamTimeRangeEnd, ( float )TotalDemoTimeInMS / 1000 );
+		}
+		else
+		{
+			UE_LOG( LogHttpReplay, Error, TEXT( "FHttpNetworkReplayStreamer::HttpDownloadFinished. FAILED." ) );
+			StreamArchive.Buffer.Empty();
+			SetLastError( ENetworkReplayError::ServiceUnavailable );
+		}
 	}
 }
 
