Merge: Ability to override request requests. This will give UT streamer the ability to add authorization headers to requests.

diff --git a/Engine/Source/Runtime/NetworkReplayStreaming/HttpNetworkReplayStreaming/Private/HttpNetworkReplayStreaming.cpp b/Engine/Source/Runtime/NetworkReplayStreaming/HttpNetworkReplayStreaming/Private/HttpNetworkReplayStreaming.cpp
--- a/Engine/Source/Runtime/NetworkReplayStreaming/HttpNetworkReplayStreaming/Private/HttpNetworkReplayStreaming.cpp
+++ b/Engine/Source/Runtime/NetworkReplayStreaming/HttpNetworkReplayStreaming/Private/HttpNetworkReplayStreaming.cpp
@@ -1317,7 +1317,8 @@ bool FHttpNetworkReplayStreamer::ProcessNextHttpRequest()
 		check( !InFlightHttpRequest.IsValid() );
 
 		InFlightHttpRequest = QueuedRequest;
-		InFlightHttpRequest->Request->ProcessRequest();
+
+		ProcessRequestInternal( InFlightHttpRequest->Request );
 
 		// We can return now since we know a http request is now in flight
 		return true;
@@ -1326,6 +1327,11 @@ bool FHttpNetworkReplayStreamer::ProcessNextHttpRequest()
 	return false;
 }
 
+void FHttpNetworkReplayStreamer::ProcessRequestInternal( TSharedPtr< class IHttpRequest > Request )
+{
+	Request->ProcessRequest();
+}
+
 void FHttpNetworkReplayStreamer::Tick( const float DeltaTime )
 {
 	if ( IsHttpRequestInFlight() )
diff --git a/Engine/Source/Runtime/NetworkReplayStreaming/HttpNetworkReplayStreaming/Public/HttpNetworkReplayStreaming.h b/Engine/Source/Runtime/NetworkReplayStreaming/HttpNetworkReplayStreaming/Public/HttpNetworkReplayStreaming.h
--- a/Engine/Source/Runtime/NetworkReplayStreaming/HttpNetworkReplayStreaming/Public/HttpNetworkReplayStreaming.h
+++ b/Engine/Source/Runtime/NetworkReplayStreaming/HttpNetworkReplayStreaming/Public/HttpNetworkReplayStreaming.h
@@ -171,6 +171,8 @@ public:
 	void EnumerateCheckpoints();
 	void ConditionallyEnumerateCheckpoints();
 
+	virtual void ProcessRequestInternal( TSharedPtr< class IHttpRequest > Request );
+
 	/** EStreamerState - Overall state of the streamer */
 	enum class EStreamerState
 	{
