The engine now notifies the replay system about players who join in progress.

--------
Integrated using branch UT_To_UE4 of change#2592351 by [REDACTED_PERSON] on 2015/06/18 14:22:24.

diff --git a/Engine/Source/Runtime/Engine/Classes/Engine/DemoNetDriver.h b/Engine/Source/Runtime/Engine/Classes/Engine/DemoNetDriver.h
--- a/Engine/Source/Runtime/Engine/Classes/Engine/DemoNetDriver.h
+++ b/Engine/Source/Runtime/Engine/Classes/Engine/DemoNetDriver.h
@@ -127,6 +127,13 @@ public:
 	void ResetDemoState();
 	void JumpToEndOfLiveReplay();
 
+	/**
+	 * Adds a join-in-progress user to the set of users associated with the currently recording replay (if any)
+	 *
+	 * @param UserString a string that uniquely identifies the user, usually his or her FUniqueNetId
+	 */
+	void AddUserToReplay(const FString& UserString);
+
 	void StopDemo();
 
 	void ReplayStreamingReady( bool bSuccess, bool bRecord );
diff --git a/Engine/Source/Runtime/Engine/Classes/Engine/GameInstance.h b/Engine/Source/Runtime/Engine/Classes/Engine/GameInstance.h
--- a/Engine/Source/Runtime/Engine/Classes/Engine/GameInstance.h
+++ b/Engine/Source/Runtime/Engine/Classes/Engine/GameInstance.h
@@ -185,6 +185,13 @@ public:
 	/** Start playing back a previously recorded replay. */
 	virtual void PlayReplay(const FString& Name);
 
+	/**
+	 * Adds a join-in-progress user to the set of users associated with the currently recording replay (if any)
+	 *
+	 * @param UserString a string that uniquely identifies the user, usually his or her FUniqueNetId
+	 */
+	virtual void AddUserToReplay(const FString& UserString);
+
 	FTimerManager* TimerManager;
 
 	/** @return online session management object associated with this game instance */
diff --git a/Engine/Source/Runtime/Engine/Private/DemoNetDriver.cpp b/Engine/Source/Runtime/Engine/Private/DemoNetDriver.cpp
--- a/Engine/Source/Runtime/Engine/Private/DemoNetDriver.cpp
+++ b/Engine/Source/Runtime/Engine/Private/DemoNetDriver.cpp
@@ -421,7 +421,12 @@ bool UDemoNetDriver::InitListen( FNetworkNotify* InNotify, FURL& ListenURL, bool
 
 	for ( int32 i = 0; i < GetWorld()->GameState->PlayerArray.Num(); i++ )
 	{
-		UserNames.Add( GetWorld()->GameState->PlayerArray[i]->UniqueId.ToString() );
+		APlayerState* PlayerState = GetWorld()->GameState->PlayerArray[i];
+
+		if ( !PlayerState->bIsABot && !PlayerState->bIsSpectator )
+		{
+			UserNames.Add( PlayerState->UniqueId.ToString() );
+		}
 	}
 
 	ReplayStreamer->StartStreaming(
@@ -1241,6 +1246,14 @@ void UDemoNetDriver::JumpToEndOfLiveReplay()
 	InitialLiveDemoTime = 0;
 }
 
+void UDemoNetDriver::AddUserToReplay( const FString& UserString )
+{
+	if ( ReplayStreamer.IsValid() )
+	{
+		ReplayStreamer->AddUserToReplay( UserString );
+	}
+}
+
 void UDemoNetDriver::TickDemoPlayback( float DeltaSeconds )
 {
 	if ( ServerConnection == NULL || ServerConnection->State == USOCK_Closed )
diff --git a/Engine/Source/Runtime/Engine/Private/GameInstance.cpp b/Engine/Source/Runtime/Engine/Private/GameInstance.cpp
--- a/Engine/Source/Runtime/Engine/Private/GameInstance.cpp
+++ b/Engine/Source/Runtime/Engine/Private/GameInstance.cpp
@@ -800,6 +800,16 @@ void UGameInstance::PlayReplay(const FString& Name)
 	}
 }
 
+void UGameInstance::AddUserToReplay(const FString& UserString)
+{
+	UWorld* CurrentWorld = GetWorld();
+
+	if ( CurrentWorld != nullptr && CurrentWorld->DemoNetDriver != nullptr )
+	{
+		CurrentWorld->DemoNetDriver->AddUserToReplay( UserString );
+	}
+}
+
 TSubclassOf<UOnlineSession> UGameInstance::GetOnlineSessionClass()
 {
 	return UOnlineSession::StaticClass();
diff --git a/Engine/Source/Runtime/Engine/Private/GameMode.cpp b/Engine/Source/Runtime/Engine/Private/GameMode.cpp
--- a/Engine/Source/Runtime/Engine/Private/GameMode.cpp
+++ b/Engine/Source/Runtime/Engine/Private/GameMode.cpp
@@ -258,6 +258,14 @@ void AGameMode::PostLogin( APlayerController* NewPlayer )
 	{
 		NewPlayer->ClientGotoState(NAME_Spectating);
 	}
+	else
+	{
+		// If NewPlayer is not only a spectator and has a valid ID, add him as a user to the replay.
+		if (NewPlayer->PlayerState->UniqueId.IsValid())
+		{
+			GetGameInstance()->AddUserToReplay(NewPlayer->PlayerState->UniqueId.ToString());
+		}
+	}
 
 	if (GameSession)
 	{
diff --git a/Engine/Source/Runtime/NetworkReplayStreaming/HttpNetworkReplayStreaming/Private/HttpNetworkReplayStreaming.cpp b/Engine/Source/Runtime/NetworkReplayStreaming/HttpNetworkReplayStreaming/Private/HttpNetworkReplayStreaming.cpp
--- a/Engine/Source/Runtime/NetworkReplayStreaming/HttpNetworkReplayStreaming/Private/HttpNetworkReplayStreaming.cpp
+++ b/Engine/Source/Runtime/NetworkReplayStreaming/HttpNetworkReplayStreaming/Private/HttpNetworkReplayStreaming.cpp
@@ -840,6 +840,36 @@ void FHttpNetworkReplayStreamer::EnumerateRecentStreams( const FNetworkReplayVer
 	AddRequestToQueue( EQueuedHttpRequestType::EnumeratingSessions, HttpRequest );
 }
 
+void FHttpNetworkReplayStreamer::AddUserToReplay(const FString& UserString)
+{
+	if ( StreamerState != EStreamerState::StreamingUp )
+	{
+		return;
+	}
+
+	if ( UserString.IsEmpty() )
+	{
+		UE_LOG(LogHttpReplay, Log, TEXT("FHttpNetworkReplayStreamer::AddUserToReplay: can't add a user with an empty UserString."));
+		return;
+	}
+
+	TSharedRef<class IHttpRequest> HttpRequest = FHttpModule::Get().CreateRequest();
+
+	HttpRequest->SetURL( FString::Printf( TEXT( "%saddusers?Session=%s" ), *ServerURL, *SessionName ) );
+	HttpRequest->SetVerb( TEXT( "POST" ) );
+
+	FNetworkReplayUserList UserList;
+
+	UserList.Users.Add(UserString);
+	FString JsonString = UserList.ToJson();
+	HttpRequest->SetContentAsString( JsonString );
+	HttpRequest->SetHeader( TEXT( "Content-Type" ), TEXT( "application/json" ) );
+
+	HttpRequest->OnProcessRequestComplete().BindRaw( this, &FHttpNetworkReplayStreamer::HttpAddUserFinished );
+
+	AddRequestToQueue( EQueuedHttpRequestType::AddingUser, HttpRequest );
+}
+
 void FHttpNetworkReplayStreamer::EnumerateCheckpoints()
 {
 	TSharedRef<class IHttpRequest> HttpRequest = FHttpModule::Get().CreateRequest();
@@ -1300,6 +1330,21 @@ void FHttpNetworkReplayStreamer::HttpEnumerateCheckpointsFinished( FHttpRequestP
 	StartStreamingDelegate = FOnStreamReadyDelegate();
 }
 
+void FHttpNetworkReplayStreamer::HttpAddUserFinished(FHttpRequestPtr HttpRequest, FHttpResponsePtr HttpResponse, bool bSucceeded)
+{
+	RequestFinished( EStreamerState::StreamingUp, EQueuedHttpRequestType::AddingUser, HttpRequest );
+
+	if ( bSucceeded && HttpResponse->GetResponseCode() == EHttpResponseCodes::Ok )
+	{		
+		UE_LOG( LogHttpReplay, Verbose, TEXT( "FHttpNetworkReplayStreamer::HttpAddUserFinished." ) );
+	}
+	else
+	{
+		// Don't consider this a fatal error
+		UE_LOG( LogHttpReplay, Warning, TEXT( "FHttpNetworkReplayStreamer::HttpAddUserFinished. FAILED" ) );
+	}
+}
+
 bool FHttpNetworkReplayStreamer::ProcessNextHttpRequest()
 {
 	if ( IsHttpRequestInFlight() )
diff --git a/Engine/Source/Runtime/NetworkReplayStreaming/HttpNetworkReplayStreaming/Public/HttpNetworkReplayStreaming.h b/Engine/Source/Runtime/NetworkReplayStreaming/HttpNetworkReplayStreaming/Public/HttpNetworkReplayStreaming.h
--- a/Engine/Source/Runtime/NetworkReplayStreaming/HttpNetworkReplayStreaming/Public/HttpNetworkReplayStreaming.h
+++ b/Engine/Source/Runtime/NetworkReplayStreaming/HttpNetworkReplayStreaming/Public/HttpNetworkReplayStreaming.h
@@ -77,7 +77,8 @@ namespace EQueuedHttpRequestType
 		EnumeratingSessions,		// We are in the process of downloading the available sessions
 		EnumeratingCheckpoints,		// We are in the process of downloading the available checkpoints
 		UploadingCheckpoint,		// We are uploading a checkpoint
-		DownloadingCheckpoint		// We are downloading a checkpoint
+		DownloadingCheckpoint,		// We are downloading a checkpoint
+		AddingUser					// We are adding a user who joined in progress during recording
 	};
 
 	inline const TCHAR* ToString( EQueuedHttpRequestType::Type Type )
@@ -108,6 +109,8 @@ namespace EQueuedHttpRequestType
 				return TEXT( "UploadingCheckpoint" );
 			case DownloadingCheckpoint:
 				return TEXT( "DownloadingCheckpoint" );
+			case AddingUser:
+				return TEXT( "AddingUser" );
 		}
 
 		return TEXT( "Unknown EQueuedHttpRequestType type." );
@@ -153,6 +156,7 @@ public:
 	virtual void		DeleteFinishedStream( const FString& StreamName, const FOnDeleteFinishedStreamComplete& Delegate ) const override;
 	virtual void		EnumerateStreams( const FNetworkReplayVersion& ReplayVersion, const FString& UserString, const FString& MetaString, const FOnEnumerateStreamsComplete& Delegate ) override;
 	virtual void		EnumerateRecentStreams( const FNetworkReplayVersion& ReplayVersion, const FString& RecentViewer, const FOnEnumerateStreamsComplete& Delegate ) override;
+	virtual void		AddUserToReplay(const FString& UserString);
 
 	virtual ENetworkReplayError::Type GetLastError() const override;
 
@@ -196,6 +200,7 @@ public:
 	void HttpUploadCheckpointFinished( FHttpRequestPtr HttpRequest, FHttpResponsePtr HttpResponse, bool bSucceeded );
 	void HttpEnumerateSessionsFinished( FHttpRequestPtr HttpRequest, FHttpResponsePtr HttpResponse, bool bSucceeded );
 	void HttpEnumerateCheckpointsFinished( FHttpRequestPtr HttpRequest, FHttpResponsePtr HttpResponse, bool bSucceeded );
+	void HttpAddUserFinished( FHttpRequestPtr HttpRequest, FHttpResponsePtr HttpResponse, bool bSucceeded );
 
 	bool ProcessNextHttpRequest();
 	void Tick( const float DeltaTime );
diff --git a/Engine/Source/Runtime/NetworkReplayStreaming/NetworkReplayStreaming/Public/NetworkReplayStreaming.h b/Engine/Source/Runtime/NetworkReplayStreaming/NetworkReplayStreaming/Public/NetworkReplayStreaming.h
--- a/Engine/Source/Runtime/NetworkReplayStreaming/NetworkReplayStreaming/Public/NetworkReplayStreaming.h
+++ b/Engine/Source/Runtime/NetworkReplayStreaming/NetworkReplayStreaming/Public/NetworkReplayStreaming.h
@@ -149,6 +149,13 @@ public:
 
 	/** Returns the last error that occurred while streaming replays */
 	virtual ENetworkReplayError::Type GetLastError() const = 0;
+
+	/**
+	 * Adds a join-in-progress user to the set of users associated with the currently recording replay (if any)
+	 *
+	 * @param UserString a string that uniquely identifies the user, usually his or her FUniqueNetId
+	 */
+	virtual void AddUserToReplay( const FString& UserString ) = 0;
 };
 
 /** Replay streamer factory */
diff --git a/Engine/Source/Runtime/NetworkReplayStreaming/NullNetworkReplayStreaming/Private/NullNetworkReplayStreaming.cpp b/Engine/Source/Runtime/NetworkReplayStreaming/NullNetworkReplayStreaming/Private/NullNetworkReplayStreaming.cpp
--- a/Engine/Source/Runtime/NetworkReplayStreaming/NullNetworkReplayStreaming/Private/NullNetworkReplayStreaming.cpp
+++ b/Engine/Source/Runtime/NetworkReplayStreaming/NullNetworkReplayStreaming/Private/NullNetworkReplayStreaming.cpp
@@ -367,6 +367,11 @@ void FNullNetworkReplayStreamer::EnumerateStreams( const FNetworkReplayVersion&
 	Delegate.ExecuteIfBound(Results);
 }
 
+void FNullNetworkReplayStreamer::AddUserToReplay(const FString& UserString)
+{
+	UE_LOG(LogNullReplay, Log, TEXT("FNullNetworkReplayStreamer::AddUserToReplay is currently unsupported."));
+}
+
 FArchive* FNullNetworkReplayStreamer::GetCheckpointArchive()
 {
 	// If the archive is null, and the API is being used properly, the caller is writing a checkpoint...
diff --git a/Engine/Source/Runtime/NetworkReplayStreaming/NullNetworkReplayStreaming/Public/NullNetworkReplayStreaming.h b/Engine/Source/Runtime/NetworkReplayStreaming/NullNetworkReplayStreaming/Public/NullNetworkReplayStreaming.h
--- a/Engine/Source/Runtime/NetworkReplayStreaming/NullNetworkReplayStreaming/Public/NullNetworkReplayStreaming.h
+++ b/Engine/Source/Runtime/NetworkReplayStreaming/NullNetworkReplayStreaming/Public/NullNetworkReplayStreaming.h
@@ -58,6 +58,7 @@ public:
 	virtual void EnumerateStreams( const FNetworkReplayVersion& ReplayVersion, const FString& UserString, const FString& MetaString, const FOnEnumerateStreamsComplete& Delegate ) override;
 	virtual void EnumerateRecentStreams( const FNetworkReplayVersion& ReplayVersion, const FString& RecentViewer, const FOnEnumerateStreamsComplete& Delegate ) override {}
 	virtual ENetworkReplayError::Type GetLastError() const override { return ENetworkReplayError::None; }
+	virtual void AddUserToReplay( const FString& UserString );
 
 private:
 	bool IsNamedStreamLive( const FString& StreamName ) const;
