Include the HTTP response code in the logs if an HTTP request to the replay service fails.

--------
Integrated using branch UT_To_UE4 of change#2599358 by [REDACTED_PERSON] on 2015/06/24 14:58:31.

diff --git a/Engine/Source/Runtime/NetworkReplayStreaming/HttpNetworkReplayStreaming/Private/HttpNetworkReplayStreaming.cpp b/Engine/Source/Runtime/NetworkReplayStreaming/HttpNetworkReplayStreaming/Private/HttpNetworkReplayStreaming.cpp
--- a/Engine/Source/Runtime/NetworkReplayStreaming/HttpNetworkReplayStreaming/Private/HttpNetworkReplayStreaming.cpp
+++ b/Engine/Source/Runtime/NetworkReplayStreaming/HttpNetworkReplayStreaming/Private/HttpNetworkReplayStreaming.cpp
@@ -924,7 +924,7 @@ void FHttpNetworkReplayStreamer::HttpStartUploadingFinished( FHttpRequestPtr Htt
 	}
 	else
 	{
-		UE_LOG( LogHttpReplay, Error, TEXT( "FHttpNetworkReplayStreamer::HttpStartUploadingFinished. FAILED" ) );
+		UE_LOG( LogHttpReplay, Error, TEXT( "FHttpNetworkReplayStreamer::HttpStartUploadingFinished. FAILED, Response code: %d" ), HttpResponse.IsValid() ? HttpResponse->GetResponseCode() : 0 );
 		SetLastError( ENetworkReplayError::ServiceUnavailable );
 	}
 }
@@ -939,7 +939,7 @@ void FHttpNetworkReplayStreamer::HttpStopUploadingFinished( FHttpRequestPtr Http
 	}
 	else
 	{
-		UE_LOG( LogHttpReplay, Error, TEXT( "FHttpNetworkReplayStreamer::HttpStopUploadingFinished. FAILED" ) );
+		UE_LOG( LogHttpReplay, Error, TEXT( "FHttpNetworkReplayStreamer::HttpStopUploadingFinished. FAILED, Response code: %d" ), HttpResponse.IsValid() ? HttpResponse->GetResponseCode() : 0 );
 		SetLastError( ENetworkReplayError::ServiceUnavailable );
 	}
 
@@ -965,7 +965,7 @@ void FHttpNetworkReplayStreamer::HttpHeaderUploadFinished( FHttpRequestPtr HttpR
 	}
 	else
 	{
-		UE_LOG( LogHttpReplay, Error, TEXT( "FHttpNetworkReplayStreamer::HttpHeaderUploadFinished. FAILED" ) );
+		UE_LOG( LogHttpReplay, Error, TEXT( "FHttpNetworkReplayStreamer::HttpHeaderUploadFinished. FAILED, Response code: %d" ), HttpResponse.IsValid() ? HttpResponse->GetResponseCode() : 0 );
 		StartStreamingDelegate.ExecuteIfBound( false, true );
 		SetLastError( ENetworkReplayError::ServiceUnavailable );
 	}
@@ -984,7 +984,7 @@ void FHttpNetworkReplayStreamer::HttpUploadStreamFinished( FHttpRequestPtr HttpR
 	}
 	else
 	{
-		UE_LOG( LogHttpReplay, Error, TEXT( "FHttpNetworkReplayStreamer::HttpUploadStreamFinished. FAILED" ) );
+		UE_LOG( LogHttpReplay, Error, TEXT( "FHttpNetworkReplayStreamer::HttpUploadStreamFinished. FAILED, Response code: %d" ), HttpResponse.IsValid() ? HttpResponse->GetResponseCode() : 0 );
 		SetLastError( ENetworkReplayError::ServiceUnavailable );
 	}
 }
@@ -999,7 +999,7 @@ void FHttpNetworkReplayStreamer::HttpUploadCheckpointFinished( FHttpRequestPtr H
 	}
 	else
 	{
-		UE_LOG( LogHttpReplay, Error, TEXT( "FHttpNetworkReplayStreamer::HttpUploadCheckpointFinished. FAILED" ) );
+		UE_LOG( LogHttpReplay, Error, TEXT( "FHttpNetworkReplayStreamer::HttpUploadCheckpointFinished. FAILED, Response code: %d" ), HttpResponse.IsValid() ? HttpResponse->GetResponseCode() : 0 );
 		SetLastError( ENetworkReplayError::ServiceUnavailable );
 	}
 }
@@ -1037,7 +1037,7 @@ void FHttpNetworkReplayStreamer::HttpStartDownloadingFinished( FHttpRequestPtr H
 	}
 	else
 	{
-		UE_LOG( LogHttpReplay, Error, TEXT( "FHttpNetworkReplayStreamer::HttpStartDownloadingFinished. FAILED. Code: %i" ), HttpResponse->GetResponseCode() );
+		UE_LOG( LogHttpReplay, Error, TEXT( "FHttpNetworkReplayStreamer::HttpStartDownloadingFinished. FAILED, Response code: %d" ), HttpResponse.IsValid() ? HttpResponse->GetResponseCode() : 0 );
 
 		StartStreamingDelegate.ExecuteIfBound( false, true );
 
@@ -1064,7 +1064,7 @@ void FHttpNetworkReplayStreamer::HttpDownloadHeaderFinished( FHttpRequestPtr Htt
 	else
 	{
 		// FAIL
-		UE_LOG( LogHttpReplay, Error, TEXT( "FHttpNetworkReplayStreamer::HttpDownloadHeaderFinished. FAILED." ) );
+		UE_LOG( LogHttpReplay, Error, TEXT( "FHttpNetworkReplayStreamer::HttpDownloadHeaderFinished. FAILED, Response code: %d" ), HttpResponse.IsValid() ? HttpResponse->GetResponseCode() : 0 );
 
 		StreamArchive.Buffer.Empty();
 		StartStreamingDelegate.ExecuteIfBound( false, false );
@@ -1127,11 +1127,11 @@ void FHttpNetworkReplayStreamer::HttpDownloadFinished( FHttpRequestPtr HttpReque
 		{
 			bStreamIsLive = false;
 
-			UE_LOG( LogHttpReplay, Warning, TEXT( "FHttpNetworkReplayStreamer::HttpDownloadFinished. Failed live, turning off live flag. Live: %i, Progress: %i / %i, Start: %i, End: %i, DemoTime: %2.2f" ), ( int )bStreamIsLive, StreamChunkIndex, NumTotalStreamChunks, ( int )StreamTimeRangeStart, ( int )StreamTimeRangeEnd, ( float )TotalDemoTimeInMS / 1000 );
+			UE_LOG( LogHttpReplay, Warning, TEXT( "FHttpNetworkReplayStreamer::HttpDownloadFinished. Failed live, turning off live flag. Response code: %d, Live: %i, Progress: %i / %i, Start: %i, End: %i, DemoTime: %2.2f" ), HttpResponse.IsValid() ? HttpResponse->GetResponseCode() : 0, ( int )bStreamIsLive, StreamChunkIndex, NumTotalStreamChunks, ( int )StreamTimeRangeStart, ( int )StreamTimeRangeEnd, ( float )TotalDemoTimeInMS / 1000 );
 		}
 		else
 		{
-			UE_LOG( LogHttpReplay, Error, TEXT( "FHttpNetworkReplayStreamer::HttpDownloadFinished. FAILED." ) );
+			UE_LOG( LogHttpReplay, Error, TEXT( "FHttpNetworkReplayStreamer::HttpDownloadFinished. FAILED, Response code: %d" ), HttpResponse.IsValid() ? HttpResponse->GetResponseCode() : 0 );
 			StreamArchive.Buffer.Empty();
 			SetLastError( ENetworkReplayError::ServiceUnavailable );
 		}
@@ -1204,7 +1204,7 @@ void FHttpNetworkReplayStreamer::HttpDownloadCheckpointFinished( FHttpRequestPtr
 	else
 	{
 		// Oops, something went wrong, notify game code of failure
-		UE_LOG( LogHttpReplay, Error, TEXT( "FHttpNetworkReplayStreamer::HttpDownloadCheckpointFinished. FAILED." ) );
+		UE_LOG( LogHttpReplay, Error, TEXT( "FHttpNetworkReplayStreamer::HttpDownloadCheckpointFinished. FAILED, Response code: %d" ), HttpResponse.IsValid() ? HttpResponse->GetResponseCode() : 0 );
 		GotoCheckpointDelegate.ExecuteIfBound( false, -1 );
 	}
 
@@ -1220,7 +1220,7 @@ void FHttpNetworkReplayStreamer::HttpRefreshViewerFinished( FHttpRequestPtr Http
 
 	if ( !bSucceeded || HttpResponse->GetResponseCode() != EHttpResponseCodes::Ok )
 	{
-		UE_LOG( LogHttpReplay, Error, TEXT( "FHttpNetworkReplayStreamer::HttpRefreshViewerFinished. FAILED." ) );
+		UE_LOG( LogHttpReplay, Error, TEXT( "FHttpNetworkReplayStreamer::HttpRefreshViewerFinished. FAILED, Response code: %d" ), HttpResponse.IsValid() ? HttpResponse->GetResponseCode() : 0 );
 		SetLastError( ENetworkReplayError::ServiceUnavailable );
 	}
 }
@@ -1269,7 +1269,7 @@ void FHttpNetworkReplayStreamer::HttpEnumerateSessionsFinished( FHttpRequestPtr
 	}
 	else
 	{
-		UE_LOG( LogHttpReplay, Warning, TEXT( "FHttpNetworkReplayStreamer::HttpEnumerateSessionsFinished. FAILED" ) );
+		UE_LOG( LogHttpReplay, Warning, TEXT( "FHttpNetworkReplayStreamer::HttpEnumerateSessionsFinished. FAILED, Response code: %d" ), HttpResponse.IsValid() ? HttpResponse->GetResponseCode() : 0 );
 		EnumerateStreamsDelegate.ExecuteIfBound( TArray<FNetworkReplayStreamInfo>() );		// FIXME: Notify failure here
 	}
 
@@ -1317,7 +1317,7 @@ void FHttpNetworkReplayStreamer::HttpEnumerateCheckpointsFinished( FHttpRequestP
 	}
 	else
 	{
-		UE_LOG( LogHttpReplay, Warning, TEXT( "FHttpNetworkReplayStreamer::HttpEnumerateCheckpointsFinished. FAILED" ) );
+		UE_LOG( LogHttpReplay, Warning, TEXT( "FHttpNetworkReplayStreamer::HttpEnumerateCheckpointsFinished. FAILED, Response code: %d" ), HttpResponse.IsValid() ? HttpResponse->GetResponseCode() : 0 );
 		StartStreamingDelegate.ExecuteIfBound( false, false );
 
 		// Reset delegate
@@ -1341,7 +1341,7 @@ void FHttpNetworkReplayStreamer::HttpAddUserFinished(FHttpRequestPtr HttpRequest
 	else
 	{
 		// Don't consider this a fatal error
-		UE_LOG( LogHttpReplay, Warning, TEXT( "FHttpNetworkReplayStreamer::HttpAddUserFinished. FAILED" ) );
+		UE_LOG( LogHttpReplay, Warning, TEXT( "FHttpNetworkReplayStreamer::HttpAddUserFinished. FAILED, Response code: %d" ), HttpResponse.IsValid() ? HttpResponse->GetResponseCode() : 0 );
 	}
 }
 
