Merge: #REPLAYS: Return an error when we abort scrubbing

diff --git a/Engine/Source/Runtime/NetworkReplayStreaming/HttpNetworkReplayStreaming/Private/HttpNetworkReplayStreaming.cpp b/Engine/Source/Runtime/NetworkReplayStreaming/HttpNetworkReplayStreaming/Private/HttpNetworkReplayStreaming.cpp
--- a/Engine/Source/Runtime/NetworkReplayStreaming/HttpNetworkReplayStreaming/Private/HttpNetworkReplayStreaming.cpp
+++ b/Engine/Source/Runtime/NetworkReplayStreaming/HttpNetworkReplayStreaming/Private/HttpNetworkReplayStreaming.cpp
@@ -485,7 +485,10 @@ void FHttpNetworkReplayStreamer::GotoTimeInMS( const uint32 TimeInMS, const FOnC
 {
 	if ( IsHttpRequestInFlight() || HasPendingHttpRequests() )
 	{
+		// If we're processing requests, be on the safe side and cancel the scrub
+		// FIXME: We can cancel the in-flight requests as well
 		UE_LOG( LogHttpReplay, Log, TEXT( "FHttpNetworkReplayStreamer::GotoTimeInMS. Busy processing pending requests." ) );
+		Delegate.ExecuteIfBound( false, -1 );
 		return;
 	}
 
