Integrate 2626341 from UE4-Orion

Integrating PeteK changes to allow custom events in replays

2623455 - Replay changes to allow custom events

diff --git a/Engine/Source/Runtime/Engine/Classes/Engine/DemoNetDriver.h b/Engine/Source/Runtime/Engine/Classes/Engine/DemoNetDriver.h
--- a/Engine/Source/Runtime/Engine/Classes/Engine/DemoNetDriver.h
+++ b/Engine/Source/Runtime/Engine/Classes/Engine/DemoNetDriver.h
@@ -2,6 +2,8 @@
 
 #pragma once
 
+#include "Runtime/Online/HTTP/Public/Interfaces/IHttpRequest.h"
+#include "Runtime/NetworkReplayStreaming/NetworkReplayStreaming/Public/NetworkReplayStreaming.h"
 #include "DemoNetDriver.generated.h"
 
 DECLARE_LOG_CATEGORY_EXTERN( LogDemo, Log, All );
@@ -131,6 +133,8 @@ public:
 	void SpawnDemoRecSpectator( UNetConnection* Connection );
 	void ResetDemoState();
 	void JumpToEndOfLiveReplay();
+	void AddEvent(const FString& Group, const FString& Meta, const TArray<uint8>& Data);
+	void EnumerateEvents(const FString& Group, FEnumerateEventsCompleteDelegate& EnumerationCompleteDelegate);
 	virtual bool IsFastForwarding() { return bIsFastForwarding; }
 
 	/**
diff --git a/Engine/Source/Runtime/Engine/Private/DemoNetDriver.cpp b/Engine/Source/Runtime/Engine/Private/DemoNetDriver.cpp
--- a/Engine/Source/Runtime/Engine/Private/DemoNetDriver.cpp
+++ b/Engine/Source/Runtime/Engine/Private/DemoNetDriver.cpp
@@ -920,6 +920,18 @@ void UDemoNetDriver::SaveCheckpoint()
 	UE_LOG( LogDemo, Log, TEXT( "Checkpoint. Total: %i, Rep size: %i, PackageMap: %u, Time: %2.2f" ), TotalSize, CheckpointSize, GuidCacheSize, CheckpointTimeInMS );
 }
 
+void UDemoNetDriver::AddEvent(const FString& Group, const FString& Meta, const TArray<uint8>& Data)
+{
+	uint32 SavedTimeMS = GetDemoCurrentTimeInMS();
+	ReplayStreamer->AddEvent(SavedTimeMS, Group, Meta, Data);
+	UE_LOG(LogDemo, Verbose, TEXT("Custom Event %s. Total: %i, Time: %2.2f"), *Group, Data.Num(), SavedTimeMS);
+}
+
+void UDemoNetDriver::EnumerateEvents(const FString& Group, FEnumerateEventsCompleteDelegate& EnumerationCompleteDelegate)
+{
+	ReplayStreamer->EnumerateEvents(Group, EnumerationCompleteDelegate);
+}
+
 void UDemoNetDriver::TickDemoRecord( float DeltaSeconds )
 {
 	if ( ClientConnections.Num() == 0 )
diff --git a/Engine/Source/Runtime/NetworkReplayStreaming/HttpNetworkReplayStreaming/Private/HttpNetworkReplayStreaming.cpp b/Engine/Source/Runtime/NetworkReplayStreaming/HttpNetworkReplayStreaming/Private/HttpNetworkReplayStreaming.cpp
--- a/Engine/Source/Runtime/NetworkReplayStreaming/HttpNetworkReplayStreaming/Private/HttpNetworkReplayStreaming.cpp
+++ b/Engine/Source/Runtime/NetworkReplayStreaming/HttpNetworkReplayStreaming/Private/HttpNetworkReplayStreaming.cpp
@@ -561,6 +561,29 @@ void FHttpNetworkReplayStreamer::FlushCheckpointInternal( uint32 TimeInMS )
 	AddRequestToQueue( EQueuedHttpRequestType::UploadingCheckpoint, HttpRequest );
 }
 
+void FHttpNetworkReplayStreamer::AddEvent( const uint32 TimeInMS, const FString& Group, const FString& Meta, const TArray<uint8>& Data )
+{
+	if (SessionName.IsEmpty() || StreamerState != EStreamerState::StreamingUp)
+	{
+		return;
+	}
+
+	// Upload a custom event to replay server
+	UE_LOG(LogHttpReplay, Verbose, TEXT("FHttpNetworkReplayStreamer::AddEvent. Size: %i, StreamChunkIndex: %i"), Data.Num(), StreamChunkIndex);
+
+	// Create the Http request and add to pending request list
+	TSharedRef<class IHttpRequest> HttpRequest = FHttpModule::Get().CreateRequest();
+
+	HttpRequest->OnProcessRequestComplete().BindRaw(this, &FHttpNetworkReplayStreamer::HttpUploadCustomEventFinished);
+
+	HttpRequest->SetURL(FString::Printf(TEXT("%suploadevent?Session=%s&Group=%s&Time1=%i&Time2=%i&Meta=%s"), *ServerURL, *SessionName, *Group, TimeInMS, TimeInMS, *Meta));
+	HttpRequest->SetVerb(TEXT("POST"));
+	HttpRequest->SetHeader(TEXT("Content-Type"), TEXT("application/octet-stream"));
+	HttpRequest->SetContent(Data);
+	
+	AddRequestToQueue(EQueuedHttpRequestType::UploadingCustomEvent, HttpRequest);
+}
+
 void FHttpNetworkReplayStreamer::DownloadHeader()
 {
 	// Download header first
@@ -905,6 +928,27 @@ void FHttpNetworkReplayStreamer::ConditionallyEnumerateCheckpoints()
 	}
 };
 
+void FHttpNetworkReplayStreamer::EnumerateEvents(const FString& Group, FEnumerateEventsCompleteDelegate& EnumerationCompleteDelegate)
+{
+	// Only have one EnumerateEventsDelegate right now
+	if (IsHttpRequestInFlight())
+	{
+		EnumerationCompleteDelegate.ExecuteIfBound(TEXT(""), false);
+		return;
+	}
+
+	TSharedRef<class IHttpRequest> HttpRequest = FHttpModule::Get().CreateRequest();
+
+	// Enumerate all of the sessions
+	HttpRequest->SetURL(FString::Printf(TEXT("%senumevents?Session=%s&Group=%s"), *ServerURL, *SessionName, *Group));
+	HttpRequest->SetVerb(TEXT("GET"));
+
+	EnumerateEventsDelegate = EnumerationCompleteDelegate;
+	HttpRequest->OnProcessRequestComplete().BindRaw(this, &FHttpNetworkReplayStreamer::HttpEnumerateEventsFinished);
+
+	AddRequestToQueue(EQueuedHttpRequestType::EnumeratingCustomEvent, HttpRequest);
+}
+
 void FHttpNetworkReplayStreamer::RequestFinished( EStreamerState ExpectedStreamerState, EQueuedHttpRequestType::Type ExpectedType, FHttpRequestPtr HttpRequest )
 {
 	check( StreamerState == ExpectedStreamerState );
@@ -1007,6 +1051,21 @@ void FHttpNetworkReplayStreamer::HttpUploadCheckpointFinished( FHttpRequestPtr H
 	}
 }
 
+void FHttpNetworkReplayStreamer::HttpUploadCustomEventFinished(FHttpRequestPtr HttpRequest, FHttpResponsePtr HttpResponse, bool bSucceeded)
+{
+	RequestFinished(EStreamerState::StreamingUp, EQueuedHttpRequestType::UploadingCustomEvent, HttpRequest);
+
+	if (bSucceeded && HttpResponse->GetResponseCode() == EHttpResponseCodes::Ok)
+	{
+		UE_LOG(LogHttpReplay, Verbose, TEXT("FHttpNetworkReplayStreamer::HttpUploadCustomEventFinished."));
+	}
+	else
+	{
+		UE_LOG(LogHttpReplay, Error, TEXT("FHttpNetworkReplayStreamer::HttpUploadCustomEventFinished. FAILED, Response code: %d"), HttpResponse.IsValid() ? HttpResponse->GetResponseCode() : 0);
+		SetLastError(ENetworkReplayError::ServiceUnavailable);
+	}
+}
+
 void FHttpNetworkReplayStreamer::HttpStartDownloadingFinished( FHttpRequestPtr HttpRequest, FHttpResponsePtr HttpResponse, bool bSucceeded )
 {
 	RequestFinished( EStreamerState::StreamingDown, EQueuedHttpRequestType::StartDownloading, HttpRequest );
@@ -1333,6 +1392,29 @@ void FHttpNetworkReplayStreamer::HttpEnumerateCheckpointsFinished( FHttpRequestP
 	StartStreamingDelegate = FOnStreamReadyDelegate();
 }
 
+void FHttpNetworkReplayStreamer::HttpEnumerateEventsFinished(FHttpRequestPtr HttpRequest, FHttpResponsePtr HttpResponse, bool bSucceeded)
+{
+	RequestFinished(EStreamerState::StreamingDown, EQueuedHttpRequestType::EnumeratingCustomEvent, HttpRequest);
+
+	if (bSucceeded && HttpResponse->GetResponseCode() == EHttpResponseCodes::Ok)
+	{
+		UE_LOG(LogHttpReplay, Verbose, TEXT("FHttpNetworkReplayStreamer::HttpEnumerateEventsFinished."));
+		FString JsonString = HttpResponse->GetContentAsString();
+		EnumerateEventsDelegate.ExecuteIfBound(JsonString, true);
+	}
+	else
+	{
+		UE_LOG(LogHttpReplay, Warning, TEXT("FHttpNetworkReplayStreamer::HttpEnumerateEventsFinished. FAILED, Response code: %d"), HttpResponse.IsValid() ? HttpResponse->GetResponseCode() : 0);
+		
+		EnumerateEventsDelegate.ExecuteIfBound(TEXT(""), false);
+
+		// Reset delegate
+		EnumerateEventsDelegate = FEnumerateEventsCompleteDelegate();
+
+		SetLastError(ENetworkReplayError::ServiceUnavailable);
+	}
+}
+
 void FHttpNetworkReplayStreamer::HttpAddUserFinished(FHttpRequestPtr HttpRequest, FHttpResponsePtr HttpResponse, bool bSucceeded)
 {
 	RequestFinished( EStreamerState::StreamingUp, EQueuedHttpRequestType::AddingUser, HttpRequest );
diff --git a/Engine/Source/Runtime/NetworkReplayStreaming/HttpNetworkReplayStreaming/Public/HttpNetworkReplayStreaming.h b/Engine/Source/Runtime/NetworkReplayStreaming/HttpNetworkReplayStreaming/Public/HttpNetworkReplayStreaming.h
--- a/Engine/Source/Runtime/NetworkReplayStreaming/HttpNetworkReplayStreaming/Public/HttpNetworkReplayStreaming.h
+++ b/Engine/Source/Runtime/NetworkReplayStreaming/HttpNetworkReplayStreaming/Public/HttpNetworkReplayStreaming.h
@@ -78,7 +78,9 @@ namespace EQueuedHttpRequestType
 		EnumeratingCheckpoints,		// We are in the process of downloading the available checkpoints
 		UploadingCheckpoint,		// We are uploading a checkpoint
 		DownloadingCheckpoint,		// We are downloading a checkpoint
-		AddingUser					// We are adding a user who joined in progress during recording
+		AddingUser,					// We are adding a user who joined in progress during recording
+		UploadingCustomEvent,		// We are uploading a custom event
+		EnumeratingCustomEvent		// We are in the process of enumerating a custom event set
 	};
 
 	inline const TCHAR* ToString( EQueuedHttpRequestType::Type Type )
@@ -111,6 +113,10 @@ namespace EQueuedHttpRequestType
 				return TEXT( "DownloadingCheckpoint" );
 			case AddingUser:
 				return TEXT( "AddingUser" );
+			case UploadingCustomEvent:
+				return TEXT( "UploadingCustomEvent" );
+			case EnumeratingCustomEvent:
+				return TEXT( "EnumeratingCustomEvent" );
 		}
 
 		return TEXT( "Unknown EQueuedHttpRequestType type." );
@@ -171,9 +177,11 @@ public:
 	void ConditionallyRefreshViewer();
 	void SetLastError( const ENetworkReplayError::Type InLastError );
 	void FlushCheckpointInternal( uint32 TimeInMS );
+	void AddEvent( const uint32 TimeInMS, const FString& Group, const FString& Meta, const TArray<uint8>& Data );
 	void AddRequestToQueue( const EQueuedHttpRequestType::Type Type, TSharedPtr< class IHttpRequest >	Request );
 	void EnumerateCheckpoints();
 	void ConditionallyEnumerateCheckpoints();
+	void EnumerateEvents( const FString& Group, FEnumerateEventsCompleteDelegate& EnumerationCompleteDelegate );
 
 	virtual void ProcessRequestInternal( TSharedPtr< class IHttpRequest > Request );
 
@@ -198,8 +206,10 @@ public:
 	void HttpHeaderUploadFinished( FHttpRequestPtr HttpRequest, FHttpResponsePtr HttpResponse, bool bSucceeded );
 	void HttpUploadStreamFinished( FHttpRequestPtr HttpRequest, FHttpResponsePtr HttpResponse, bool bSucceeded );
 	void HttpUploadCheckpointFinished( FHttpRequestPtr HttpRequest, FHttpResponsePtr HttpResponse, bool bSucceeded );
+	void HttpUploadCustomEventFinished( FHttpRequestPtr HttpRequest, FHttpResponsePtr HttpResponse, bool bSucceeded );
 	void HttpEnumerateSessionsFinished( FHttpRequestPtr HttpRequest, FHttpResponsePtr HttpResponse, bool bSucceeded );
 	void HttpEnumerateCheckpointsFinished( FHttpRequestPtr HttpRequest, FHttpResponsePtr HttpResponse, bool bSucceeded );
+	void HttpEnumerateEventsFinished( FHttpRequestPtr HttpRequest, FHttpResponsePtr HttpResponse, bool bSucceeded );
 	void HttpAddUserFinished( FHttpRequestPtr HttpRequest, FHttpResponsePtr HttpResponse, bool bSucceeded );
 
 	bool ProcessNextHttpRequest();
@@ -235,6 +245,7 @@ public:
 	FOnStreamReadyDelegate			StartStreamingDelegate;		// Delegate passed in to StartStreaming
 	FOnEnumerateStreamsComplete		EnumerateStreamsDelegate;
 	FOnCheckpointReadyDelegate		GotoCheckpointDelegate;
+	FEnumerateEventsCompleteDelegate EnumerateEventsDelegate;
 	int32							DownloadCheckpointIndex;
 	int64							LastGotoTimeInMS;
 
diff --git a/Engine/Source/Runtime/NetworkReplayStreaming/NetworkReplayStreaming/Public/NetworkReplayStreaming.h b/Engine/Source/Runtime/NetworkReplayStreaming/NetworkReplayStreaming/Public/NetworkReplayStreaming.h
--- a/Engine/Source/Runtime/NetworkReplayStreaming/NetworkReplayStreaming/Public/NetworkReplayStreaming.h
+++ b/Engine/Source/Runtime/NetworkReplayStreaming/NetworkReplayStreaming/Public/NetworkReplayStreaming.h
@@ -89,6 +89,8 @@ DECLARE_DELEGATE_OneParam( FOnDeleteFinishedStreamComplete, const bool );
  */
 DECLARE_DELEGATE_OneParam( FOnEnumerateStreamsComplete, const TArray<FNetworkReplayStreamInfo>& );
 
+DECLARE_DELEGATE_TwoParams(FEnumerateEventsCompleteDelegate, const FString&, bool);
+
 class FNetworkReplayVersion
 {
 public:
@@ -121,7 +123,9 @@ public:
 	virtual void SetHighPriorityTimeRange( const uint32 StartTimeInMS, const uint32 EndTimeInMS ) = 0;
 	virtual bool IsDataAvailableForTimeRange( const uint32 StartTimeInMS, const uint32 EndTimeInMS ) = 0;
 	virtual bool IsLoadingCheckpoint() const = 0;
-	
+	virtual void AddEvent( const uint32 TimeInMS, const FString& Group, const FString& Meta, const TArray<uint8>& Data ) = 0;
+	virtual void EnumerateEvents( const FString& Group, FEnumerateEventsCompleteDelegate& EnumerationCompleteDelegate ) = 0;
+
 	/** Returns true if the playing stream is currently in progress */
 	virtual bool IsLive() const = 0;
 
diff --git a/Engine/Source/Runtime/NetworkReplayStreaming/NullNetworkReplayStreaming/Private/NullNetworkReplayStreaming.cpp b/Engine/Source/Runtime/NetworkReplayStreaming/NullNetworkReplayStreaming/Private/NullNetworkReplayStreaming.cpp
--- a/Engine/Source/Runtime/NetworkReplayStreaming/NullNetworkReplayStreaming/Private/NullNetworkReplayStreaming.cpp
+++ b/Engine/Source/Runtime/NetworkReplayStreaming/NullNetworkReplayStreaming/Private/NullNetworkReplayStreaming.cpp
@@ -372,6 +372,16 @@ void FNullNetworkReplayStreamer::AddUserToReplay(const FString& UserString)
 	UE_LOG(LogNullReplay, Log, TEXT("FNullNetworkReplayStreamer::AddUserToReplay is currently unsupported."));
 }
 
+void FNullNetworkReplayStreamer::AddEvent( const uint32 TimeInMS, const FString& Group, const FString& Meta, const TArray<uint8>& Data )
+{
+	UE_LOG(LogNullReplay, Log, TEXT("FNullNetworkReplayStreamer::AddEvent is currently unsupported."));
+}
+
+void FNullNetworkReplayStreamer::EnumerateEvents( const FString& Group, FEnumerateEventsCompleteDelegate& EnumerationCompleteDelegate )
+{
+	UE_LOG(LogNullReplay, Log, TEXT("FNullNetworkReplayStreamer::EnumerateEvents is currently unsupported."));
+}
+
 FArchive* FNullNetworkReplayStreamer::GetCheckpointArchive()
 {
 	// If the archive is null, and the API is being used properly, the caller is writing a checkpoint...
diff --git a/Engine/Source/Runtime/NetworkReplayStreaming/NullNetworkReplayStreaming/Public/NullNetworkReplayStreaming.h b/Engine/Source/Runtime/NetworkReplayStreaming/NullNetworkReplayStreaming/Public/NullNetworkReplayStreaming.h
--- a/Engine/Source/Runtime/NetworkReplayStreaming/NullNetworkReplayStreaming/Public/NullNetworkReplayStreaming.h
+++ b/Engine/Source/Runtime/NetworkReplayStreaming/NullNetworkReplayStreaming/Public/NullNetworkReplayStreaming.h
@@ -59,6 +59,8 @@ public:
 	virtual void EnumerateRecentStreams( const FNetworkReplayVersion& ReplayVersion, const FString& RecentViewer, const FOnEnumerateStreamsComplete& Delegate ) override {}
 	virtual ENetworkReplayError::Type GetLastError() const override { return ENetworkReplayError::None; }
 	virtual void AddUserToReplay( const FString& UserString );
+	virtual void AddEvent( const uint32 TimeInMS, const FString& Group, const FString& Meta, const TArray<uint8>& Data );
+	virtual void EnumerateEvents( const FString& Group, FEnumerateEventsCompleteDelegate& EnumerationCompleteDelegate );
 
 private:
 	bool IsNamedStreamLive( const FString& StreamName ) const;
