Merging replay streaming changes from Orion branch

2620313 - Fix issue where we would get disconnected from the replay server if we paused for too long. Fixes OR-2753.

2628891 - First pass on bug replay events:
- pressing J when in a game saves a bug event
- loading that replay enumerates all bug events and adds them to the timeline bar

2630827 - Bugs on the replay timeline are clickable and will teleport the camera to the location where the bug was entered. Code is messy and needs cleanup.

2638789 - Add text box to replay list widget that filters replays by an event group

diff --git a/Engine/Source/Runtime/NetworkReplayStreaming/HttpNetworkReplayStreaming/Private/HttpNetworkReplayStreaming.cpp b/Engine/Source/Runtime/NetworkReplayStreaming/HttpNetworkReplayStreaming/Private/HttpNetworkReplayStreaming.cpp
--- a/Engine/Source/Runtime/NetworkReplayStreaming/HttpNetworkReplayStreaming/Private/HttpNetworkReplayStreaming.cpp
+++ b/Engine/Source/Runtime/NetworkReplayStreaming/HttpNetworkReplayStreaming/Private/HttpNetworkReplayStreaming.cpp
@@ -460,7 +460,7 @@ void FHttpNetworkReplayStreamer::GotoCheckpointIndex( const int32 CheckpointInde
 		return;
 	}
 
-	if ( !CheckpointList.Checkpoints.IsValidIndex( CheckpointIndex ) )
+	if ( !CheckpointList.ReplayEvents.IsValidIndex( CheckpointIndex ) )
 	{
 		UE_LOG( LogHttpReplay, Warning, TEXT( "FHttpNetworkReplayStreamer::GotoCheckpointIndex. Invalid checkpoint index." ) );
 		Delegate.ExecuteIfBound( false, -1 );
@@ -470,7 +470,7 @@ void FHttpNetworkReplayStreamer::GotoCheckpointIndex( const int32 CheckpointInde
 	TSharedRef<class IHttpRequest> HttpRequest = FHttpModule::Get().CreateRequest();
 
 	// Download the next stream chunk
-	HttpRequest->SetURL( FString::Printf( TEXT( "%sdownloadevent?Session=%s&ID=%s" ), *ServerURL, *SessionName, *CheckpointList.Checkpoints[CheckpointIndex].ID ) );
+	HttpRequest->SetURL( FString::Printf( TEXT( "%sdownloadevent?Session=%s&ID=%s" ), *ServerURL, *SessionName, *CheckpointList.ReplayEvents[CheckpointIndex].ID ) );
 	HttpRequest->SetVerb( TEXT( "GET" ) );
 
 	HttpRequest->OnProcessRequestComplete().BindRaw( this, &FHttpNetworkReplayStreamer::HttpDownloadCheckpointFinished );
@@ -481,6 +481,61 @@ void FHttpNetworkReplayStreamer::GotoCheckpointIndex( const int32 CheckpointInde
 	AddRequestToQueue( EQueuedHttpRequestType::DownloadingCheckpoint, HttpRequest );
 }
 
+void FHttpNetworkReplayStreamer::SearchEvents(const FString& EventGroup, const FOnEnumerateStreamsComplete& Delegate)
+{
+	TSharedRef<class IHttpRequest> HttpRequest = FHttpModule::Get().CreateRequest();
+
+	HttpRequest->SetURL(FString::Printf(TEXT("%ssearchevents?group=%s"), *ServerURL, *EventGroup));
+	HttpRequest->SetVerb(TEXT("GET"));
+
+	HttpRequest->OnProcessRequestComplete().BindRaw(this, &FHttpNetworkReplayStreamer::HttpEnumerateSessionsFinished);
+
+	EnumerateStreamsDelegate = Delegate;
+
+	AddRequestToQueue(EQueuedHttpRequestType::EnumeratingSessions, HttpRequest);
+}
+
+void FHttpNetworkReplayStreamer::RequestEventData(const FString& EventID, const FOnRequestEventDataComplete& Delegate)
+{
+	TSharedRef<class IHttpRequest> HttpRequest = FHttpModule::Get().CreateRequest();
+
+	// Download the next stream chunk
+	HttpRequest->SetURL(FString::Printf(TEXT("%sdownloadevent?Session=%s&ID=%s"), *ServerURL, *SessionName, *EventID));
+	HttpRequest->SetVerb(TEXT("GET"));
+
+	HttpRequest->OnProcessRequestComplete().BindRaw(this, &FHttpNetworkReplayStreamer::HttpRequestEventDataFinished);
+
+	RequestEventDataCompleteDelegate = Delegate;
+
+	AddRequestToQueue(EQueuedHttpRequestType::RequestEventData, HttpRequest);
+}
+
+void FHttpNetworkReplayStreamer::HttpRequestEventDataFinished(FHttpRequestPtr HttpRequest, FHttpResponsePtr HttpResponse, bool bSucceeded)
+{
+	RequestFinished(EStreamerState::StreamingDown, EQueuedHttpRequestType::RequestEventData, HttpRequest);
+
+	check(RequestEventDataCompleteDelegate.IsBound());
+
+	if (bSucceeded && HttpResponse->GetResponseCode() == EHttpResponseCodes::Ok)
+	{
+		// Use a temp variable to hold the delegate because we want the delegate to be able to start the next request.
+		FOnRequestEventDataComplete Delegate = RequestEventDataCompleteDelegate;
+		RequestEventDataCompleteDelegate = FOnRequestEventDataComplete();
+		Delegate.ExecuteIfBound(HttpResponse->GetContent(), true);
+		UE_LOG(LogHttpReplay, Log, TEXT("FHttpNetworkReplayStreamer::HttpRequestEventDataFinished."));
+	}
+	else
+	{
+		// FAIL
+		UE_LOG(LogHttpReplay, Error, TEXT("FHttpNetworkReplayStreamer::HttpRequestEventDataFinished. FAILED, Response code: %d"), HttpResponse.IsValid() ? HttpResponse->GetResponseCode() : 0);
+
+		RequestEventDataCompleteDelegate.ExecuteIfBound(TArray<uint8>(), false);
+		RequestEventDataCompleteDelegate = FOnRequestEventDataComplete();
+
+		SetLastError(ENetworkReplayError::ServiceUnavailable);
+	}
+}
+
 void FHttpNetworkReplayStreamer::GotoTimeInMS( const uint32 TimeInMS, const FOnCheckpointReadyDelegate& Delegate )
 {
 	if ( IsHttpRequestInFlight() || HasPendingHttpRequests() )
@@ -509,19 +564,19 @@ void FHttpNetworkReplayStreamer::GotoTimeInMS( const uint32 TimeInMS, const FOnC
 
 	LastGotoTimeInMS = FMath::Min( TimeInMS, TotalDemoTimeInMS );
 
-	if ( CheckpointList.Checkpoints.Num() > 0 && TimeInMS >= CheckpointList.Checkpoints[ CheckpointList.Checkpoints.Num() - 1 ].Time1 )
+	if ( CheckpointList.ReplayEvents.Num() > 0 && TimeInMS >= CheckpointList.ReplayEvents[ CheckpointList.ReplayEvents.Num() - 1 ].Time1 )
 	{
 		// If we're after the very last checkpoint, that's the one we want
-		CheckpointIndex = CheckpointList.Checkpoints.Num() - 1;
+		CheckpointIndex = CheckpointList.ReplayEvents.Num() - 1;
 	}
 	else
 	{
 		// Checkpoints should be sorted by time, return the checkpoint that exists right before the current time
 		// For fine scrubbing, we'll fast forward the rest of the way
 		// NOTE - If we're right before the very first checkpoint, we'll return -1, which is what we want when we want to start from the very beginning
-		for ( int i = 0; i < CheckpointList.Checkpoints.Num(); i++ )
+		for ( int i = 0; i < CheckpointList.ReplayEvents.Num(); i++ )
 		{
-			if ( TimeInMS < CheckpointList.Checkpoints[i].Time1 )
+			if ( TimeInMS < CheckpointList.ReplayEvents[i].Time1 )
 			{
 				CheckpointIndex = i - 1;
 				break;
@@ -928,12 +983,12 @@ void FHttpNetworkReplayStreamer::ConditionallyEnumerateCheckpoints()
 	}
 };
 
-void FHttpNetworkReplayStreamer::EnumerateEvents(const FString& Group, FEnumerateEventsCompleteDelegate& EnumerationCompleteDelegate)
+void FHttpNetworkReplayStreamer::EnumerateEvents(const FString& Group, const FEnumerateEventsCompleteDelegate& EnumerationCompleteDelegate)
 {
 	// Only have one EnumerateEventsDelegate right now
 	if (IsHttpRequestInFlight())
 	{
-		EnumerationCompleteDelegate.ExecuteIfBound(TEXT(""), false);
+		EnumerationCompleteDelegate.ExecuteIfBound(FReplayEventList(), false);
 		return;
 	}
 
@@ -1237,25 +1292,25 @@ void FHttpNetworkReplayStreamer::HttpDownloadCheckpointFinished( FHttpRequestPtr
 		StreamTimeRangeEnd		= 0;
 
 		// Set the next chunk to be right after this checkpoint (which was stored in the metadata)
-		StreamChunkIndex = FCString::Atoi( *CheckpointList.Checkpoints[ DownloadCheckpointIndex ].Metadata );
+		StreamChunkIndex = FCString::Atoi( *CheckpointList.ReplayEvents[ DownloadCheckpointIndex ].Metadata );
 
 		// If we want to fast forward past the end of a stream, clamp to the checkpoint
 		if ( LastGotoTimeInMS >= 0 && StreamChunkIndex >= NumTotalStreamChunks )
 		{
 			UE_LOG( LogHttpReplay, Warning, TEXT( "FHttpNetworkReplayStreamer::HttpDownloadCheckpointFinished. Clamped to checkpoint: %i" ), LastGotoTimeInMS );
 
-			StreamTimeRangeStart	= CheckpointList.Checkpoints[DownloadCheckpointIndex].Time1;
-			StreamTimeRangeEnd		= CheckpointList.Checkpoints[DownloadCheckpointIndex].Time1;
+			StreamTimeRangeStart	= CheckpointList.ReplayEvents[DownloadCheckpointIndex].Time1;
+			StreamTimeRangeEnd		= CheckpointList.ReplayEvents[DownloadCheckpointIndex].Time1;
 			LastGotoTimeInMS		= -1;
 		}
 
 		if ( LastGotoTimeInMS >= 0 )
 		{
 			// If we are fine scrubbing, make sure to wait on the part of the stream that is needed to do this in one frame
-			SetHighPriorityTimeRange( CheckpointList.Checkpoints[ DownloadCheckpointIndex ].Time1, LastGotoTimeInMS );
+			SetHighPriorityTimeRange( CheckpointList.ReplayEvents[ DownloadCheckpointIndex ].Time1, LastGotoTimeInMS );
 			
 			// Subtract off checkpoint time so we pass in the leftover to the engine to fast forward through for the fine scrubbing part
-			LastGotoTimeInMS -= CheckpointList.Checkpoints[ DownloadCheckpointIndex ].Time1;
+			LastGotoTimeInMS -= CheckpointList.ReplayEvents[ DownloadCheckpointIndex ].Time1;
 		}
 
 		// Notify game code of success
@@ -1348,7 +1403,7 @@ void FHttpNetworkReplayStreamer::HttpEnumerateCheckpointsFinished( FHttpRequestP
 
 		FString JsonString = HttpResponse->GetContentAsString();
 
-		CheckpointList.Checkpoints.Empty();
+		CheckpointList.ReplayEvents.Empty();
 
 		if ( !CheckpointList.FromJson( JsonString ) )
 		{
@@ -1358,22 +1413,20 @@ void FHttpNetworkReplayStreamer::HttpEnumerateCheckpointsFinished( FHttpRequestP
 
 			// Reset delegate
 			StartStreamingDelegate = FOnStreamReadyDelegate();
-
 			SetLastError( ENetworkReplayError::ServiceUnavailable );
-
 			return;
 		}
 
 		// Sort checkpoints by time
 		struct FCompareCheckpointTime
 		{
-			FORCEINLINE bool operator()( const FCheckpointListItem & A, const FCheckpointListItem & B ) const
+			FORCEINLINE bool operator()(const FReplayEventListItem & A, const FReplayEventListItem & B) const
 			{
 				return A.Time1 < B.Time1;
 			}
 		};
 
-		Sort( CheckpointList.Checkpoints.GetData(), CheckpointList.Checkpoints.Num(), FCompareCheckpointTime() );
+		Sort( CheckpointList.ReplayEvents.GetData(), CheckpointList.ReplayEvents.Num(), FCompareCheckpointTime() );
 
 		StartStreamingDelegate.ExecuteIfBound( true, false );
 	}
@@ -1398,15 +1451,30 @@ void FHttpNetworkReplayStreamer::HttpEnumerateEventsFinished(FHttpRequestPtr Htt
 
 	if (bSucceeded && HttpResponse->GetResponseCode() == EHttpResponseCodes::Ok)
 	{
-		UE_LOG(LogHttpReplay, Verbose, TEXT("FHttpNetworkReplayStreamer::HttpEnumerateEventsFinished."));
+		FReplayEventList EventList;
 		FString JsonString = HttpResponse->GetContentAsString();
-		EnumerateEventsDelegate.ExecuteIfBound(JsonString, true);
+		if (!EventList.FromJson(JsonString))
+		{
+			UE_LOG(LogHttpReplay, Warning, TEXT("FHttpNetworkReplayStreamer::HttpEnumerateEventsFinished. FromJson FAILED"));
+
+			EnumerateEventsDelegate.ExecuteIfBound(FReplayEventList(), false);
+
+			// Reset delegate
+			EnumerateEventsDelegate = FEnumerateEventsCompleteDelegate();
+
+			SetLastError(ENetworkReplayError::ServiceUnavailable);
+
+			return;
+		}
+
+		UE_LOG(LogHttpReplay, Verbose, TEXT("FHttpNetworkReplayStreamer::HttpEnumerateEventsFinished. %s"), *JsonString);
+		EnumerateEventsDelegate.ExecuteIfBound(EventList, true);
 	}
 	else
 	{
 		UE_LOG(LogHttpReplay, Warning, TEXT("FHttpNetworkReplayStreamer::HttpEnumerateEventsFinished. FAILED, Response code: %d"), HttpResponse.IsValid() ? HttpResponse->GetResponseCode() : 0);
 		
-		EnumerateEventsDelegate.ExecuteIfBound(TEXT(""), false);
+		EnumerateEventsDelegate.ExecuteIfBound(FReplayEventList(), false);
 
 		// Reset delegate
 		EnumerateEventsDelegate = FEnumerateEventsCompleteDelegate();
@@ -1581,6 +1649,11 @@ bool FHttpNetworkReplayStreamingFactory::IsTickable() const
 	return true;
 }
 
+bool FHttpNetworkReplayStreamingFactory::IsTickableWhenPaused() const
+{
+	return true;
+}
+
 TStatId FHttpNetworkReplayStreamingFactory::GetStatId() const
 {
 	RETURN_QUICK_DECLARE_CYCLE_STAT( FHttpNetworkReplayStreamingFactory, STATGROUP_Tickables );
diff --git a/Engine/Source/Runtime/NetworkReplayStreaming/HttpNetworkReplayStreaming/Public/HttpNetworkReplayStreaming.h b/Engine/Source/Runtime/NetworkReplayStreaming/HttpNetworkReplayStreaming/Public/HttpNetworkReplayStreaming.h
--- a/Engine/Source/Runtime/NetworkReplayStreaming/HttpNetworkReplayStreaming/Public/HttpNetworkReplayStreaming.h
+++ b/Engine/Source/Runtime/NetworkReplayStreaming/HttpNetworkReplayStreaming/Public/HttpNetworkReplayStreaming.h
@@ -6,43 +6,6 @@
 #include "Runtime/Engine/Public/Tickable.h"
 #include "OnlineJsonSerializer.h"
 
-class FCheckpointListItem : public FOnlineJsonSerializable
-{
-public:
-	FCheckpointListItem() {}
-	virtual ~FCheckpointListItem() {}
-
-	FString		ID;
-	FString		Group;
-	FString		Metadata;
-	uint32		Time1;
-	uint32		Time2;
-
-	// FOnlineJsonSerializable
-	BEGIN_ONLINE_JSON_SERIALIZER
-		ONLINE_JSON_SERIALIZE( "id",			ID );
-		ONLINE_JSON_SERIALIZE( "group",			Group );
-		ONLINE_JSON_SERIALIZE( "meta",			Metadata );
-		ONLINE_JSON_SERIALIZE( "time1",			Time1 );
-		ONLINE_JSON_SERIALIZE( "time2",			Time2 );
-	END_ONLINE_JSON_SERIALIZER
-};
-
-class FCheckpointList : public FOnlineJsonSerializable
-{
-public:
-	FCheckpointList()
-	{}
-	virtual ~FCheckpointList() {}
-
-	TArray< FCheckpointListItem > Checkpoints;
-
-	// FOnlineJsonSerializable
-	BEGIN_ONLINE_JSON_SERIALIZER
-		ONLINE_JSON_SERIALIZE_ARRAY_SERIALIZABLE( "events", Checkpoints, FCheckpointListItem );
-	END_ONLINE_JSON_SERIALIZER
-};
-
 /**
  * Archive used to buffer stream over http
  */
@@ -80,7 +43,8 @@ namespace EQueuedHttpRequestType
 		DownloadingCheckpoint,		// We are downloading a checkpoint
 		AddingUser,					// We are adding a user who joined in progress during recording
 		UploadingCustomEvent,		// We are uploading a custom event
-		EnumeratingCustomEvent		// We are in the process of enumerating a custom event set
+		EnumeratingCustomEvent,		// We are in the process of enumerating a custom event set
+		RequestEventData,			// We are in the process of requesting the data for a specific event
 	};
 
 	inline const TCHAR* ToString( EQueuedHttpRequestType::Type Type )
@@ -117,6 +81,8 @@ namespace EQueuedHttpRequestType
 				return TEXT( "UploadingCustomEvent" );
 			case EnumeratingCustomEvent:
 				return TEXT( "EnumeratingCustomEvent" );
+			case RequestEventData:
+				return TEXT("RequestEventData");
 		}
 
 		return TEXT( "Unknown EQueuedHttpRequestType type." );
@@ -163,7 +129,8 @@ public:
 	virtual void		EnumerateStreams( const FNetworkReplayVersion& ReplayVersion, const FString& UserString, const FString& MetaString, const FOnEnumerateStreamsComplete& Delegate ) override;
 	virtual void		EnumerateRecentStreams( const FNetworkReplayVersion& ReplayVersion, const FString& RecentViewer, const FOnEnumerateStreamsComplete& Delegate ) override;
 	virtual void		AddUserToReplay(const FString& UserString);
-
+	virtual void		RequestEventData(const FString& EventId, const FOnRequestEventDataComplete& Delegate) override;
+	virtual void		SearchEvents(const FString& EventGroup, const FOnEnumerateStreamsComplete& Delegate) override;
 	virtual ENetworkReplayError::Type GetLastError() const override;
 
 	/** FHttpNetworkReplayStreamer */
@@ -181,7 +148,7 @@ public:
 	void AddRequestToQueue( const EQueuedHttpRequestType::Type Type, TSharedPtr< class IHttpRequest >	Request );
 	void EnumerateCheckpoints();
 	void ConditionallyEnumerateCheckpoints();
-	void EnumerateEvents( const FString& Group, FEnumerateEventsCompleteDelegate& EnumerationCompleteDelegate );
+	void EnumerateEvents( const FString& Group, const FEnumerateEventsCompleteDelegate& EnumerationCompleteDelegate );
 
 	virtual void ProcessRequestInternal( TSharedPtr< class IHttpRequest > Request );
 
@@ -211,6 +178,7 @@ public:
 	void HttpEnumerateCheckpointsFinished( FHttpRequestPtr HttpRequest, FHttpResponsePtr HttpResponse, bool bSucceeded );
 	void HttpEnumerateEventsFinished( FHttpRequestPtr HttpRequest, FHttpResponsePtr HttpResponse, bool bSucceeded );
 	void HttpAddUserFinished( FHttpRequestPtr HttpRequest, FHttpResponsePtr HttpResponse, bool bSucceeded );
+	void HttpRequestEventDataFinished(FHttpRequestPtr HttpRequest, FHttpResponsePtr HttpResponse, bool bSucceeded);
 
 	bool ProcessNextHttpRequest();
 	void Tick( const float DeltaTime );
@@ -246,10 +214,11 @@ public:
 	FOnEnumerateStreamsComplete		EnumerateStreamsDelegate;
 	FOnCheckpointReadyDelegate		GotoCheckpointDelegate;
 	FEnumerateEventsCompleteDelegate EnumerateEventsDelegate;
+	FOnRequestEventDataComplete		RequestEventDataCompleteDelegate;
 	int32							DownloadCheckpointIndex;
 	int64							LastGotoTimeInMS;
 
-	FCheckpointList					CheckpointList;
+	FReplayEventList					CheckpointList;
 
 	TQueue< TSharedPtr< FQueuedHttpRequest > >	QueuedHttpRequests;
 	TSharedPtr< FQueuedHttpRequest >			InFlightHttpRequest;
@@ -265,6 +234,7 @@ public:
 	virtual void Tick( float DeltaTime ) override;
 	virtual bool IsTickable() const override;
 	virtual TStatId GetStatId() const override;
+	bool IsTickableWhenPaused() const override;
 
 	TArray< TSharedPtr< FHttpNetworkReplayStreamer > > HttpStreamers;
 };
diff --git a/Engine/Source/Runtime/NetworkReplayStreaming/NetworkReplayStreaming/NetworkReplayStreaming.Build.cs b/Engine/Source/Runtime/NetworkReplayStreaming/NetworkReplayStreaming/NetworkReplayStreaming.Build.cs
--- a/Engine/Source/Runtime/NetworkReplayStreaming/NetworkReplayStreaming/NetworkReplayStreaming.Build.cs
+++ b/Engine/Source/Runtime/NetworkReplayStreaming/NetworkReplayStreaming/NetworkReplayStreaming.Build.cs
@@ -8,10 +8,13 @@ namespace UnrealBuildTool.Rules
 		{
 			PrivateIncludePaths.Add( "Runtime/NetworkReplayStreaming/NetworkReplayStreaming/Private" );
 
+            PrivateIncludePathModuleNames.Add("OnlineSubsystem");
+
 			PublicDependencyModuleNames.AddRange(
 				new string[]
 				{
 					"Core",
+                    "Json",
                     // ... add other public dependencies that you statically link with here ...
 				}
 			);
diff --git a/Engine/Source/Runtime/NetworkReplayStreaming/NetworkReplayStreaming/Public/NetworkReplayStreaming.h b/Engine/Source/Runtime/NetworkReplayStreaming/NetworkReplayStreaming/Public/NetworkReplayStreaming.h
--- a/Engine/Source/Runtime/NetworkReplayStreaming/NetworkReplayStreaming/Public/NetworkReplayStreaming.h
+++ b/Engine/Source/Runtime/NetworkReplayStreaming/NetworkReplayStreaming/Public/NetworkReplayStreaming.h
@@ -5,6 +5,44 @@
 // Dependencies.
 #include "ModuleManager.h"
 #include "Core.h"
+#include "OnlineJsonSerializer.h"
+
+class FReplayEventListItem : public FOnlineJsonSerializable
+{
+public:
+	FReplayEventListItem() {}
+	virtual ~FReplayEventListItem() {}
+
+	FString		ID;
+	FString		Group;
+	FString		Metadata;
+	uint32		Time1;
+	uint32		Time2;
+
+	// FOnlineJsonSerializable
+	BEGIN_ONLINE_JSON_SERIALIZER
+	ONLINE_JSON_SERIALIZE("id", ID);
+	ONLINE_JSON_SERIALIZE("group", Group);
+	ONLINE_JSON_SERIALIZE("meta", Metadata);
+	ONLINE_JSON_SERIALIZE("time1", Time1);
+	ONLINE_JSON_SERIALIZE("time2", Time2);
+	END_ONLINE_JSON_SERIALIZER
+};
+
+class FReplayEventList : public FOnlineJsonSerializable
+{
+public:
+	FReplayEventList()
+	{}
+	virtual ~FReplayEventList() {}
+
+	TArray< FReplayEventListItem > ReplayEvents;
+
+	// FOnlineJsonSerializable
+	BEGIN_ONLINE_JSON_SERIALIZER
+		ONLINE_JSON_SERIALIZE_ARRAY_SERIALIZABLE("events", ReplayEvents, FReplayEventListItem);
+	END_ONLINE_JSON_SERIALIZER
+};
 
 /** Struct to store information about a stream, returned from search results. */
 struct FNetworkReplayStreamInfo
@@ -89,7 +127,21 @@ DECLARE_DELEGATE_OneParam( FOnDeleteFinishedStreamComplete, const bool );
  */
 DECLARE_DELEGATE_OneParam( FOnEnumerateStreamsComplete, const TArray<FNetworkReplayStreamInfo>& );
 
-DECLARE_DELEGATE_TwoParams(FEnumerateEventsCompleteDelegate, const FString&, bool);
+/**
+* Delegate called when EnumerateEvents() completes.
+*
+* @param ReplayEventList A list of events that were found
+# @param bWasSuccessful Whether enumerating events was successful
+*/
+DECLARE_DELEGATE_TwoParams(FEnumerateEventsCompleteDelegate, const FReplayEventList&, bool);
+
+/**
+* Delegate called when RequestEventData() completes.
+*
+* @param ReplayEventListItem A replay event with its data parameter filled in
+# @param bWasSuccessful Whether enumerating events was successful
+*/
+DECLARE_DELEGATE_TwoParams(FOnRequestEventDataComplete, const TArray<uint8>&, bool)
 
 class FNetworkReplayVersion
 {
@@ -124,8 +176,9 @@ public:
 	virtual bool IsDataAvailableForTimeRange( const uint32 StartTimeInMS, const uint32 EndTimeInMS ) = 0;
 	virtual bool IsLoadingCheckpoint() const = 0;
 	virtual void AddEvent( const uint32 TimeInMS, const FString& Group, const FString& Meta, const TArray<uint8>& Data ) = 0;
-	virtual void EnumerateEvents( const FString& Group, FEnumerateEventsCompleteDelegate& EnumerationCompleteDelegate ) = 0;
-
+	virtual void EnumerateEvents( const FString& Group, const FEnumerateEventsCompleteDelegate& EnumerationCompleteDelegate ) = 0;
+	virtual void RequestEventData(const FString& EventID, const FOnRequestEventDataComplete& RequestEventDataComplete) = 0;
+	virtual void SearchEvents(const FString& EventGroup, const FOnEnumerateStreamsComplete& Delegate) = 0;
 	/** Returns true if the playing stream is currently in progress */
 	virtual bool IsLive() const = 0;
 
diff --git a/Engine/Source/Runtime/NetworkReplayStreaming/NullNetworkReplayStreaming/Private/NullNetworkReplayStreaming.cpp b/Engine/Source/Runtime/NetworkReplayStreaming/NullNetworkReplayStreaming/Private/NullNetworkReplayStreaming.cpp
--- a/Engine/Source/Runtime/NetworkReplayStreaming/NullNetworkReplayStreaming/Private/NullNetworkReplayStreaming.cpp
+++ b/Engine/Source/Runtime/NetworkReplayStreaming/NullNetworkReplayStreaming/Private/NullNetworkReplayStreaming.cpp
@@ -377,11 +377,21 @@ void FNullNetworkReplayStreamer::AddEvent( const uint32 TimeInMS, const FString&
 	UE_LOG(LogNullReplay, Log, TEXT("FNullNetworkReplayStreamer::AddEvent is currently unsupported."));
 }
 
-void FNullNetworkReplayStreamer::EnumerateEvents( const FString& Group, FEnumerateEventsCompleteDelegate& EnumerationCompleteDelegate )
+void FNullNetworkReplayStreamer::EnumerateEvents( const FString& Group, const FEnumerateEventsCompleteDelegate& EnumerationCompleteDelegate )
 {
 	UE_LOG(LogNullReplay, Log, TEXT("FNullNetworkReplayStreamer::EnumerateEvents is currently unsupported."));
 }
 
+void FNullNetworkReplayStreamer::RequestEventData(const FString& EventID, const FOnRequestEventDataComplete& RequestEventDataComplete)
+{
+	UE_LOG(LogNullReplay, Log, TEXT("FNullNetworkReplayStreamer::RequestEventData is currently unsupported."));
+}
+
+void FNullNetworkReplayStreamer::SearchEvents(const FString& EventGroup, const FOnEnumerateStreamsComplete& Delegate)
+{
+	UE_LOG(LogNullReplay, Log, TEXT("FNullNetworkReplayStreamer::SearchEvents is currently unsupported."));
+}
+
 FArchive* FNullNetworkReplayStreamer::GetCheckpointArchive()
 {
 	// If the archive is null, and the API is being used properly, the caller is writing a checkpoint...
diff --git a/Engine/Source/Runtime/NetworkReplayStreaming/NullNetworkReplayStreaming/Public/NullNetworkReplayStreaming.h b/Engine/Source/Runtime/NetworkReplayStreaming/NullNetworkReplayStreaming/Public/NullNetworkReplayStreaming.h
--- a/Engine/Source/Runtime/NetworkReplayStreaming/NullNetworkReplayStreaming/Public/NullNetworkReplayStreaming.h
+++ b/Engine/Source/Runtime/NetworkReplayStreaming/NullNetworkReplayStreaming/Public/NullNetworkReplayStreaming.h
@@ -58,9 +58,11 @@ public:
 	virtual void EnumerateStreams( const FNetworkReplayVersion& ReplayVersion, const FString& UserString, const FString& MetaString, const FOnEnumerateStreamsComplete& Delegate ) override;
 	virtual void EnumerateRecentStreams( const FNetworkReplayVersion& ReplayVersion, const FString& RecentViewer, const FOnEnumerateStreamsComplete& Delegate ) override {}
 	virtual ENetworkReplayError::Type GetLastError() const override { return ENetworkReplayError::None; }
-	virtual void AddUserToReplay( const FString& UserString );
-	virtual void AddEvent( const uint32 TimeInMS, const FString& Group, const FString& Meta, const TArray<uint8>& Data );
-	virtual void EnumerateEvents( const FString& Group, FEnumerateEventsCompleteDelegate& EnumerationCompleteDelegate );
+	virtual void AddUserToReplay(const FString& UserString) override;
+	virtual void AddEvent(const uint32 TimeInMS, const FString& Group, const FString& Meta, const TArray<uint8>& Data) override;
+	virtual void EnumerateEvents(const FString& Group, const FEnumerateEventsCompleteDelegate& EnumerationCompleteDelegate) override;
+	virtual void RequestEventData(const FString& EventID, const FOnRequestEventDataComplete& RequestEventDataComplete) override;
+	virtual void SearchEvents(const FString& EventGroup, const FOnEnumerateStreamsComplete& Delegate) override;
 
 private:
 	bool IsNamedStreamLive( const FString& StreamName ) const;
