Http replay streamer should queue event enumeration requests

Merging 2643483

//depot/UE4-UT/...

to //depot/UE4/...

diff --git a/Engine/Source/Runtime/NetworkReplayStreaming/HttpNetworkReplayStreaming/Private/HttpNetworkReplayStreaming.cpp b/Engine/Source/Runtime/NetworkReplayStreaming/HttpNetworkReplayStreaming/Private/HttpNetworkReplayStreaming.cpp
--- a/Engine/Source/Runtime/NetworkReplayStreaming/HttpNetworkReplayStreaming/Private/HttpNetworkReplayStreaming.cpp
+++ b/Engine/Source/Runtime/NetworkReplayStreaming/HttpNetworkReplayStreaming/Private/HttpNetworkReplayStreaming.cpp
@@ -985,12 +985,6 @@ void FHttpNetworkReplayStreamer::ConditionallyEnumerateCheckpoints()
 
 void FHttpNetworkReplayStreamer::EnumerateEvents(const FString& Group, const FEnumerateEventsCompleteDelegate& EnumerationCompleteDelegate)
 {
-	// Only have one EnumerateEventsDelegate right now
-	if (IsHttpRequestInFlight())
-	{
-		EnumerationCompleteDelegate.ExecuteIfBound(FReplayEventList(), false);
-		return;
-	}
 
 	TSharedRef<class IHttpRequest> HttpRequest = FHttpModule::Get().CreateRequest();
 
@@ -998,8 +992,7 @@ void FHttpNetworkReplayStreamer::EnumerateEvents(const FString& Group, const FEn
 	HttpRequest->SetURL(FString::Printf(TEXT("%senumevents?Session=%s&Group=%s"), *ServerURL, *SessionName, *Group));
 	HttpRequest->SetVerb(TEXT("GET"));
 
-	EnumerateEventsDelegate = EnumerationCompleteDelegate;
-	HttpRequest->OnProcessRequestComplete().BindRaw(this, &FHttpNetworkReplayStreamer::HttpEnumerateEventsFinished);
+	HttpRequest->OnProcessRequestComplete().BindRaw(this, &FHttpNetworkReplayStreamer::HttpEnumerateEventsFinished, EnumerationCompleteDelegate);
 
 	AddRequestToQueue(EQueuedHttpRequestType::EnumeratingCustomEvent, HttpRequest);
 }
@@ -1445,7 +1438,7 @@ void FHttpNetworkReplayStreamer::HttpEnumerateCheckpointsFinished( FHttpRequestP
 	StartStreamingDelegate = FOnStreamReadyDelegate();
 }
 
-void FHttpNetworkReplayStreamer::HttpEnumerateEventsFinished(FHttpRequestPtr HttpRequest, FHttpResponsePtr HttpResponse, bool bSucceeded)
+void FHttpNetworkReplayStreamer::HttpEnumerateEventsFinished(FHttpRequestPtr HttpRequest, FHttpResponsePtr HttpResponse, bool bSucceeded, FEnumerateEventsCompleteDelegate EnumerateEventsDelegate)
 {
 	RequestFinished(EStreamerState::StreamingDown, EQueuedHttpRequestType::EnumeratingCustomEvent, HttpRequest);
 
@@ -1475,10 +1468,7 @@ void FHttpNetworkReplayStreamer::HttpEnumerateEventsFinished(FHttpRequestPtr Htt
 		UE_LOG(LogHttpReplay, Warning, TEXT("FHttpNetworkReplayStreamer::HttpEnumerateEventsFinished. FAILED, Response code: %d"), HttpResponse.IsValid() ? HttpResponse->GetResponseCode() : 0);
 		
 		EnumerateEventsDelegate.ExecuteIfBound(FReplayEventList(), false);
-
-		// Reset delegate
-		EnumerateEventsDelegate = FEnumerateEventsCompleteDelegate();
-
+		
 		SetLastError(ENetworkReplayError::ServiceUnavailable);
 	}
 }
diff --git a/Engine/Source/Runtime/NetworkReplayStreaming/HttpNetworkReplayStreaming/Public/HttpNetworkReplayStreaming.h b/Engine/Source/Runtime/NetworkReplayStreaming/HttpNetworkReplayStreaming/Public/HttpNetworkReplayStreaming.h
--- a/Engine/Source/Runtime/NetworkReplayStreaming/HttpNetworkReplayStreaming/Public/HttpNetworkReplayStreaming.h
+++ b/Engine/Source/Runtime/NetworkReplayStreaming/HttpNetworkReplayStreaming/Public/HttpNetworkReplayStreaming.h
@@ -176,7 +176,7 @@ public:
 	void HttpUploadCustomEventFinished( FHttpRequestPtr HttpRequest, FHttpResponsePtr HttpResponse, bool bSucceeded );
 	void HttpEnumerateSessionsFinished( FHttpRequestPtr HttpRequest, FHttpResponsePtr HttpResponse, bool bSucceeded );
 	void HttpEnumerateCheckpointsFinished( FHttpRequestPtr HttpRequest, FHttpResponsePtr HttpResponse, bool bSucceeded );
-	void HttpEnumerateEventsFinished( FHttpRequestPtr HttpRequest, FHttpResponsePtr HttpResponse, bool bSucceeded );
+	void HttpEnumerateEventsFinished( FHttpRequestPtr HttpRequest, FHttpResponsePtr HttpResponse, bool bSucceeded, FEnumerateEventsCompleteDelegate EnumerateEventsDelegate );
 	void HttpAddUserFinished( FHttpRequestPtr HttpRequest, FHttpResponsePtr HttpResponse, bool bSucceeded );
 	void HttpRequestEventDataFinished(FHttpRequestPtr HttpRequest, FHttpResponsePtr HttpResponse, bool bSucceeded);
 
@@ -213,7 +213,6 @@ public:
 	FOnStreamReadyDelegate			StartStreamingDelegate;		// Delegate passed in to StartStreaming
 	FOnEnumerateStreamsComplete		EnumerateStreamsDelegate;
 	FOnCheckpointReadyDelegate		GotoCheckpointDelegate;
-	FEnumerateEventsCompleteDelegate EnumerateEventsDelegate;
 	FOnRequestEventDataComplete		RequestEventDataCompleteDelegate;
 	int32							DownloadCheckpointIndex;
 	int64							LastGotoTimeInMS;
