#replays, merge: Don't send version/cl if they are blank, MCP doesn't like that

diff --git a/Engine/Source/Runtime/NetworkReplayStreaming/HttpNetworkReplayStreaming/Private/HttpNetworkReplayStreaming.cpp b/Engine/Source/Runtime/NetworkReplayStreaming/HttpNetworkReplayStreaming/Private/HttpNetworkReplayStreaming.cpp
--- a/Engine/Source/Runtime/NetworkReplayStreaming/HttpNetworkReplayStreaming/Private/HttpNetworkReplayStreaming.cpp
+++ b/Engine/Source/Runtime/NetworkReplayStreaming/HttpNetworkReplayStreaming/Private/HttpNetworkReplayStreaming.cpp
@@ -873,7 +873,18 @@ void FHttpNetworkReplayStreamer::EnumerateStreams( const FNetworkReplayVersion&
 	TSharedRef<class IHttpRequest> HttpRequest = FHttpModule::Get().CreateRequest();
 
 	// Build base URL
-	FString URL = FString::Printf( TEXT( "%senumsessions?App=%s&Version=%u&CL=%u" ), *ServerURL, *InReplayVersion.AppString, InReplayVersion.NetworkVersion, InReplayVersion.Changelist );
+	FString URL = FString::Printf( TEXT( "%senumsessions?App=%s" ), *ServerURL, *InReplayVersion.AppString );
+
+	// Add optional stuff
+	if ( InReplayVersion.Changelist != 0 )
+	{
+		URL += FString::Printf( TEXT( "&CL=%u" ), InReplayVersion.Changelist );
+	}
+
+	if ( InReplayVersion.NetworkVersion != 0 )
+	{
+		URL += FString::Printf( TEXT( "&Version=%u" ), InReplayVersion.NetworkVersion );
+	}
 
 	const FString MetaStringToUse = !CVarMetaFilterOverride.GetValueOnGameThread().IsEmpty() ? CVarMetaFilterOverride.GetValueOnGameThread() : MetaString;
 
