Make RequestEventData queue safe

Merging 2648244

//depot/UE4-UT/...

to //depot/UE4/...

diff --git a/Engine/Source/Runtime/NetworkReplayStreaming/HttpNetworkReplayStreaming/Private/HttpNetworkReplayStreaming.cpp b/Engine/Source/Runtime/NetworkReplayStreaming/HttpNetworkReplayStreaming/Private/HttpNetworkReplayStreaming.cpp
--- a/Engine/Source/Runtime/NetworkReplayStreaming/HttpNetworkReplayStreaming/Private/HttpNetworkReplayStreaming.cpp
+++ b/Engine/Source/Runtime/NetworkReplayStreaming/HttpNetworkReplayStreaming/Private/HttpNetworkReplayStreaming.cpp
@@ -503,14 +503,12 @@ void FHttpNetworkReplayStreamer::RequestEventData(const FString& EventID, const
 	HttpRequest->SetURL(FString::Printf(TEXT("%sdownloadevent?Session=%s&ID=%s"), *ServerURL, *SessionName, *EventID));
 	HttpRequest->SetVerb(TEXT("GET"));
 
-	HttpRequest->OnProcessRequestComplete().BindRaw(this, &FHttpNetworkReplayStreamer::HttpRequestEventDataFinished);
-
-	RequestEventDataCompleteDelegate = Delegate;
-
+	HttpRequest->OnProcessRequestComplete().BindRaw(this, &FHttpNetworkReplayStreamer::HttpRequestEventDataFinished, Delegate);
+	
 	AddRequestToQueue(EQueuedHttpRequestType::RequestEventData, HttpRequest);
 }
 
-void FHttpNetworkReplayStreamer::HttpRequestEventDataFinished(FHttpRequestPtr HttpRequest, FHttpResponsePtr HttpResponse, bool bSucceeded)
+void FHttpNetworkReplayStreamer::HttpRequestEventDataFinished(FHttpRequestPtr HttpRequest, FHttpResponsePtr HttpResponse, bool bSucceeded, FOnRequestEventDataComplete RequestEventDataCompleteDelegate)
 {
 	RequestFinished(EStreamerState::StreamingDown, EQueuedHttpRequestType::RequestEventData, HttpRequest);
 
@@ -518,10 +516,7 @@ void FHttpNetworkReplayStreamer::HttpRequestEventDataFinished(FHttpRequestPtr Ht
 
 	if (bSucceeded && HttpResponse->GetResponseCode() == EHttpResponseCodes::Ok)
 	{
-		// Use a temp variable to hold the delegate because we want the delegate to be able to start the next request.
-		FOnRequestEventDataComplete Delegate = RequestEventDataCompleteDelegate;
-		RequestEventDataCompleteDelegate = FOnRequestEventDataComplete();
-		Delegate.ExecuteIfBound(HttpResponse->GetContent(), true);
+		RequestEventDataCompleteDelegate.ExecuteIfBound(HttpResponse->GetContent(), true);
 		UE_LOG(LogHttpReplay, Log, TEXT("FHttpNetworkReplayStreamer::HttpRequestEventDataFinished."));
 	}
 	else
@@ -530,7 +525,6 @@ void FHttpNetworkReplayStreamer::HttpRequestEventDataFinished(FHttpRequestPtr Ht
 		UE_LOG(LogHttpReplay, Error, TEXT("FHttpNetworkReplayStreamer::HttpRequestEventDataFinished. FAILED, Response code: %d"), HttpResponse.IsValid() ? HttpResponse->GetResponseCode() : 0);
 
 		RequestEventDataCompleteDelegate.ExecuteIfBound(TArray<uint8>(), false);
-		RequestEventDataCompleteDelegate = FOnRequestEventDataComplete();
 
 		SetLastError(ENetworkReplayError::ServiceUnavailable);
 	}
diff --git a/Engine/Source/Runtime/NetworkReplayStreaming/HttpNetworkReplayStreaming/Public/HttpNetworkReplayStreaming.h b/Engine/Source/Runtime/NetworkReplayStreaming/HttpNetworkReplayStreaming/Public/HttpNetworkReplayStreaming.h
--- a/Engine/Source/Runtime/NetworkReplayStreaming/HttpNetworkReplayStreaming/Public/HttpNetworkReplayStreaming.h
+++ b/Engine/Source/Runtime/NetworkReplayStreaming/HttpNetworkReplayStreaming/Public/HttpNetworkReplayStreaming.h
@@ -178,7 +178,7 @@ public:
 	void HttpEnumerateCheckpointsFinished( FHttpRequestPtr HttpRequest, FHttpResponsePtr HttpResponse, bool bSucceeded );
 	void HttpEnumerateEventsFinished( FHttpRequestPtr HttpRequest, FHttpResponsePtr HttpResponse, bool bSucceeded, FEnumerateEventsCompleteDelegate EnumerateEventsDelegate );
 	void HttpAddUserFinished( FHttpRequestPtr HttpRequest, FHttpResponsePtr HttpResponse, bool bSucceeded );
-	void HttpRequestEventDataFinished(FHttpRequestPtr HttpRequest, FHttpResponsePtr HttpResponse, bool bSucceeded);
+	void HttpRequestEventDataFinished( FHttpRequestPtr HttpRequest, FHttpResponsePtr HttpResponse, bool bSucceeded, FOnRequestEventDataComplete RequestEventDataCompleteDelegate );
 
 	bool ProcessNextHttpRequest();
 	void Tick( const float DeltaTime );
@@ -213,7 +213,6 @@ public:
 	FOnStreamReadyDelegate			StartStreamingDelegate;		// Delegate passed in to StartStreaming
 	FOnEnumerateStreamsComplete		EnumerateStreamsDelegate;
 	FOnCheckpointReadyDelegate		GotoCheckpointDelegate;
-	FOnRequestEventDataComplete		RequestEventDataCompleteDelegate;
 	int32							DownloadCheckpointIndex;
 	int64							LastGotoTimeInMS;
 
