Event data accessor for demo net driver

Merging 2648403

//depot/UE4-UT/...

to //depot/UE4/...

diff --git a/Engine/Source/Runtime/Engine/Classes/Engine/DemoNetDriver.h b/Engine/Source/Runtime/Engine/Classes/Engine/DemoNetDriver.h
--- a/Engine/Source/Runtime/Engine/Classes/Engine/DemoNetDriver.h
+++ b/Engine/Source/Runtime/Engine/Classes/Engine/DemoNetDriver.h
@@ -135,6 +135,7 @@ public:
 	void JumpToEndOfLiveReplay();
 	void AddEvent(const FString& Group, const FString& Meta, const TArray<uint8>& Data);
 	void EnumerateEvents(const FString& Group, FEnumerateEventsCompleteDelegate& EnumerationCompleteDelegate);
+	void RequestEventData(const FString& EventID, FOnRequestEventDataComplete& RequestEventDataCompleteDelegate);
 	virtual bool IsFastForwarding() { return bIsFastForwarding; }
 
 	/**
diff --git a/Engine/Source/Runtime/Engine/Private/DemoNetDriver.cpp b/Engine/Source/Runtime/Engine/Private/DemoNetDriver.cpp
--- a/Engine/Source/Runtime/Engine/Private/DemoNetDriver.cpp
+++ b/Engine/Source/Runtime/Engine/Private/DemoNetDriver.cpp
@@ -970,6 +970,14 @@ void UDemoNetDriver::EnumerateEvents(const FString& Group, FEnumerateEventsCompl
 	ReplayStreamer->EnumerateEvents(Group, EnumerationCompleteDelegate);
 }
 
+void UDemoNetDriver::RequestEventData(const FString& EventID, FOnRequestEventDataComplete& RequestEventDataCompleteDelegate)
+{
+	if (ReplayStreamer.IsValid())
+	{
+		ReplayStreamer->RequestEventData(EventID, RequestEventDataCompleteDelegate);
+	}
+}
+
 void UDemoNetDriver::TickDemoRecord( float DeltaSeconds )
 {
 	if ( ClientConnections.Num() == 0 )
diff --git a/Engine/Source/Runtime/NetworkReplayStreaming/HttpNetworkReplayStreaming/Private/HttpNetworkReplayStreaming.cpp b/Engine/Source/Runtime/NetworkReplayStreaming/HttpNetworkReplayStreaming/Private/HttpNetworkReplayStreaming.cpp
--- a/Engine/Source/Runtime/NetworkReplayStreaming/HttpNetworkReplayStreaming/Private/HttpNetworkReplayStreaming.cpp
+++ b/Engine/Source/Runtime/NetworkReplayStreaming/HttpNetworkReplayStreaming/Private/HttpNetworkReplayStreaming.cpp
@@ -517,7 +517,7 @@ void FHttpNetworkReplayStreamer::HttpRequestEventDataFinished(FHttpRequestPtr Ht
 	if (bSucceeded && HttpResponse->GetResponseCode() == EHttpResponseCodes::Ok)
 	{
 		RequestEventDataCompleteDelegate.ExecuteIfBound(HttpResponse->GetContent(), true);
-		UE_LOG(LogHttpReplay, Log, TEXT("FHttpNetworkReplayStreamer::HttpRequestEventDataFinished."));
+		UE_LOG(LogHttpReplay, Verbose, TEXT("FHttpNetworkReplayStreamer::HttpRequestEventDataFinished."));
 	}
 	else
 	{
