Allow creating event during replay downstreaming

Merging 2651042

//depot/UE4-UT/...

to //depot/UE4/...

diff --git a/Engine/Source/Runtime/NetworkReplayStreaming/HttpNetworkReplayStreaming/Private/HttpNetworkReplayStreaming.cpp b/Engine/Source/Runtime/NetworkReplayStreaming/HttpNetworkReplayStreaming/Private/HttpNetworkReplayStreaming.cpp
--- a/Engine/Source/Runtime/NetworkReplayStreaming/HttpNetworkReplayStreaming/Private/HttpNetworkReplayStreaming.cpp
+++ b/Engine/Source/Runtime/NetworkReplayStreaming/HttpNetworkReplayStreaming/Private/HttpNetworkReplayStreaming.cpp
@@ -612,7 +612,7 @@ void FHttpNetworkReplayStreamer::FlushCheckpointInternal( uint32 TimeInMS )
 
 void FHttpNetworkReplayStreamer::AddEvent( const uint32 TimeInMS, const FString& Group, const FString& Meta, const TArray<uint8>& Data )
 {
-	if (SessionName.IsEmpty() || StreamerState != EStreamerState::StreamingUp)
+	if (SessionName.IsEmpty() || (StreamerState != EStreamerState::StreamingUp && StreamerState != EStreamerState::StreamingDown))
 	{
 		return;
 	}
@@ -1106,7 +1106,7 @@ void FHttpNetworkReplayStreamer::HttpUploadCheckpointFinished( FHttpRequestPtr H
 
 void FHttpNetworkReplayStreamer::HttpUploadCustomEventFinished(FHttpRequestPtr HttpRequest, FHttpResponsePtr HttpResponse, bool bSucceeded)
 {
-	RequestFinished(EStreamerState::StreamingUp, EQueuedHttpRequestType::UploadingCustomEvent, HttpRequest);
+	RequestFinished(StreamerState, EQueuedHttpRequestType::UploadingCustomEvent, HttpRequest);
 
 	if (bSucceeded && HttpResponse->GetResponseCode() == EHttpResponseCodes::Ok)
 	{
@@ -1115,7 +1115,12 @@ void FHttpNetworkReplayStreamer::HttpUploadCustomEventFinished(FHttpRequestPtr H
 	else
 	{
 		UE_LOG(LogHttpReplay, Error, TEXT("FHttpNetworkReplayStreamer::HttpUploadCustomEventFinished. FAILED, Response code: %d"), HttpResponse.IsValid() ? HttpResponse->GetResponseCode() : 0);
-		SetLastError(ENetworkReplayError::ServiceUnavailable);
+
+		// Don't care about the error coming down
+		if (StreamerState == EStreamerState::StreamingUp)
+		{
+			SetLastError(ENetworkReplayError::ServiceUnavailable);
+		}
 	}
 }
 
