#merge, replays: Specifically detect an already in-flight scrub operation instead of checking for any in-flight command. This allows scrubbing to happen if something like a chunk download is in-flight, which shouldn't interfere.

diff --git a/Engine/Source/Runtime/NetworkReplayStreaming/HttpNetworkReplayStreaming/Private/HttpNetworkReplayStreaming.cpp b/Engine/Source/Runtime/NetworkReplayStreaming/HttpNetworkReplayStreaming/Private/HttpNetworkReplayStreaming.cpp
--- a/Engine/Source/Runtime/NetworkReplayStreaming/HttpNetworkReplayStreaming/Private/HttpNetworkReplayStreaming.cpp
+++ b/Engine/Source/Runtime/NetworkReplayStreaming/HttpNetworkReplayStreaming/Private/HttpNetworkReplayStreaming.cpp
@@ -428,7 +428,7 @@ void FHttpNetworkReplayStreamer::FlushCheckpoint( const uint32 TimeInMS )
 
 void FHttpNetworkReplayStreamer::GotoCheckpointIndex( const int32 CheckpointIndex, const FOnCheckpointReadyDelegate& Delegate )
 {
-	if ( GotoCheckpointDelegate.IsBound() )
+	if ( GotoCheckpointDelegate.IsBound() || DownloadCheckpointIndex != -1 )
 	{
 		// If we're currently going to a checkpoint now, ignore this request
 		UE_LOG( LogHttpReplay, Warning, TEXT( "FHttpNetworkReplayStreamer::GotoCheckpointIndex. Busy processing another checkpoint." ) );
@@ -532,7 +532,7 @@ void FHttpNetworkReplayStreamer::HttpRequestEventDataFinished(FHttpRequestPtr Ht
 
 void FHttpNetworkReplayStreamer::GotoTimeInMS( const uint32 TimeInMS, const FOnCheckpointReadyDelegate& Delegate )
 {
-	if ( IsHttpRequestInFlight() || HasPendingHttpRequests() )
+	if ( LastGotoTimeInMS != -1 || DownloadCheckpointIndex != -1 )
 	{
 		// If we're processing requests, be on the safe side and cancel the scrub
 		// FIXME: We can cancel the in-flight requests as well
