#merge, replays:  Failing to download event data isn't catrastophic + fix issue where calling the next request can re-use the event delegate

diff --git a/Engine/Source/Runtime/NetworkReplayStreaming/HttpNetworkReplayStreaming/Private/HttpNetworkReplayStreaming.cpp b/Engine/Source/Runtime/NetworkReplayStreaming/HttpNetworkReplayStreaming/Private/HttpNetworkReplayStreaming.cpp
--- a/Engine/Source/Runtime/NetworkReplayStreaming/HttpNetworkReplayStreaming/Private/HttpNetworkReplayStreaming.cpp
+++ b/Engine/Source/Runtime/NetworkReplayStreaming/HttpNetworkReplayStreaming/Private/HttpNetworkReplayStreaming.cpp
@@ -525,8 +525,6 @@ void FHttpNetworkReplayStreamer::HttpRequestEventDataFinished(FHttpRequestPtr Ht
 		UE_LOG(LogHttpReplay, Error, TEXT("FHttpNetworkReplayStreamer::HttpRequestEventDataFinished. FAILED, Response code: %d"), HttpResponse.IsValid() ? HttpResponse->GetResponseCode() : 0);
 
 		RequestEventDataCompleteDelegate.ExecuteIfBound(TArray<uint8>(), false);
-
-		SetLastError(ENetworkReplayError::ServiceUnavailable);
 	}
 }
 
