#merge, replays:  Convert a few steps into tasks to help ensure order of operations

diff --git a/Engine/Source/Runtime/NetworkReplayStreaming/HttpNetworkReplayStreaming/Private/HttpNetworkReplayStreaming.cpp b/Engine/Source/Runtime/NetworkReplayStreaming/HttpNetworkReplayStreaming/Private/HttpNetworkReplayStreaming.cpp
--- a/Engine/Source/Runtime/NetworkReplayStreaming/HttpNetworkReplayStreaming/Private/HttpNetworkReplayStreaming.cpp
+++ b/Engine/Source/Runtime/NetworkReplayStreaming/HttpNetworkReplayStreaming/Private/HttpNetworkReplayStreaming.cpp
@@ -129,7 +129,6 @@ FHttpNetworkReplayStreamer::FHttpNetworkReplayStreamer() :
 	LastRefreshCheckpointTime( 0 ),
 	StreamerState( EStreamerState::Idle ), 
 	bStopStreamingCalled( false ), 
-	bNeedToUploadHeader( false ),
 	bStreamIsLive( false ),
 	NumTotalStreamChunks( 0 ),
 	TotalDemoTimeInMS( 0 ),
@@ -232,9 +231,6 @@ void FHttpNetworkReplayStreamer::StartStreaming( const FString& CustomName, cons
 
 		SessionName.Empty();
 
-		// We can't upload the header until we have the session name, which depends on the StartUploading task above to finish
-		bNeedToUploadHeader = true;
-
 		FString MetaString;
 
 		if ( FParse::Value( FCommandLine::Get(), TEXT( "ReplayMeta=" ), MetaString ) && !MetaString.IsEmpty() )
@@ -262,10 +258,13 @@ void FHttpNetworkReplayStreamer::StartStreaming( const FString& CustomName, cons
 		}
 
 		AddRequestToQueue( EQueuedHttpRequestType::StartUploading, HttpRequest );
+		
+		// We need to upload the header AFTER StartUploading is done (so we have session name)
+		AddRequestToQueue( EQueuedHttpRequestType::UploadHeader, nullptr );
 	}
 }
 
-void FHttpNetworkReplayStreamer::AddRequestToQueue( const EQueuedHttpRequestType::Type Type, TSharedPtr< class IHttpRequest >	Request )
+void FHttpNetworkReplayStreamer::AddRequestToQueue( const EQueuedHttpRequestType::Type Type, TSharedPtr< class IHttpRequest > Request )
 {
 	UE_LOG( LogHttpReplay, Verbose, TEXT( "FHttpNetworkReplayStreamer::AddRequestToQueue. Type: %s" ), EQueuedHttpRequestType::ToString( Type ) );
 
@@ -300,6 +299,8 @@ void FHttpNetworkReplayStreamer::StopStreaming()
 		// Send one last http request to stop uploading		
 		StopUploading();
 	}
+
+	AddRequestToQueue( EQueuedHttpRequestType::StopStreaming, nullptr );
 }
 
 void FHttpNetworkReplayStreamer::UploadHeader()
@@ -348,7 +349,7 @@ void FHttpNetworkReplayStreamer::FlushStream()
 {
 	check( StreamArchive.ArIsSaving );
 
-	if ( bNeedToUploadHeader )
+	if ( SessionName.IsEmpty() )
 	{
 		// If we haven't uploaded the header, or we are not recording, we don't need to flush
 		UE_LOG( LogHttpReplay, Warning, TEXT( "FHttpNetworkReplayStreamer::FlushStream. Waiting on header upload." ) );
@@ -770,7 +771,7 @@ FArchive* FHttpNetworkReplayStreamer::GetStreamingArchive()
 
 FArchive* FHttpNetworkReplayStreamer::GetCheckpointArchive()
 {	
-	if ( bNeedToUploadHeader )
+	if ( SessionName.IsEmpty() )
 	{
 		// If we need to upload the header, we're not ready to save checkpoints
 		// NOTE - The code needs to be resilient to this, and keep trying!!!!
@@ -1512,6 +1513,21 @@ bool FHttpNetworkReplayStreamer::ProcessNextHttpRequest()
 
 		check( !InFlightHttpRequest.IsValid() );
 
+		// Check for a couple special requests that aren't really http calls, but just using the request system for order of operations management
+		if ( QueuedRequest->Type == EQueuedHttpRequestType::UploadHeader )
+		{
+			check( !SessionName.IsEmpty() );
+			UploadHeader();
+			return ProcessNextHttpRequest();
+		}
+		else if ( QueuedRequest->Type == EQueuedHttpRequestType::StopStreaming )
+		{
+			check( IsStreaming() );
+			StreamerState = EStreamerState::Idle;
+			bStopStreamingCalled = false;
+			return ProcessNextHttpRequest();
+		}
+
 		InFlightHttpRequest = QueuedRequest;
 
 		ProcessRequestInternal( InFlightHttpRequest->Request );
@@ -1547,29 +1563,6 @@ void FHttpNetworkReplayStreamer::Tick( const float DeltaTime )
 	// We should have no pending or requests in flight at this point
 	check( !HasPendingHttpRequests() );
 
-	// See if we're waiting on uploading the header
-	// We have to do it this way, because the header depends on the session name, and we don't have that until
-	if ( bNeedToUploadHeader )
-	{
-		// We can't upload the header until we have the session name, so keep waiting on it
-		if ( !SessionName.IsEmpty() )
-		{
-			UploadHeader();
-			bNeedToUploadHeader = false;
-		}
-
-		return;
-	}
-
-	if ( bStopStreamingCalled )
-	{
-		// If we have a pending stop streaming request, we can fully stop once all http requests are processed
-		check( IsStreaming() );
-		StreamerState = EStreamerState::Idle;
-		bStopStreamingCalled = false;
-		return;
-	}
-
 	if ( StreamerState == EStreamerState::StreamingUp )
 	{
 		ConditionallyFlushStream();
diff --git a/Engine/Source/Runtime/NetworkReplayStreaming/HttpNetworkReplayStreaming/Public/HttpNetworkReplayStreaming.h b/Engine/Source/Runtime/NetworkReplayStreaming/HttpNetworkReplayStreaming/Public/HttpNetworkReplayStreaming.h
--- a/Engine/Source/Runtime/NetworkReplayStreaming/HttpNetworkReplayStreaming/Public/HttpNetworkReplayStreaming.h
+++ b/Engine/Source/Runtime/NetworkReplayStreaming/HttpNetworkReplayStreaming/Public/HttpNetworkReplayStreaming.h
@@ -45,6 +45,8 @@ namespace EQueuedHttpRequestType
 		UploadingCustomEvent,		// We are uploading a custom event
 		EnumeratingCustomEvent,		// We are in the process of enumerating a custom event set
 		RequestEventData,			// We are in the process of requesting the data for a specific event
+		UploadHeader,				// Request to upload header (has to be done after we get info from server)
+		StopStreaming,				// Request to stop streaming
 	};
 
 	inline const TCHAR* ToString( EQueuedHttpRequestType::Type Type )
@@ -83,6 +85,10 @@ namespace EQueuedHttpRequestType
 				return TEXT( "EnumeratingCustomEvent" );
 			case RequestEventData:
 				return TEXT("RequestEventData");
+			case UploadHeader:
+				return TEXT( "UploadHeader" );
+			case StopStreaming:
+				return TEXT( "StopStreaming" );
 		}
 
 		return TEXT( "Unknown EQueuedHttpRequestType type." );
@@ -198,7 +204,6 @@ public:
 	double					LastRefreshCheckpointTime;
 	EStreamerState			StreamerState;			// Overall state of the streamer
 	bool					bStopStreamingCalled;
-	bool					bNeedToUploadHeader;	// We're waiting on session name so we can upload header
 	bool					bStreamIsLive;			// If true, we are viewing a live stream
 	int32					NumTotalStreamChunks;
 	uint32					TotalDemoTimeInMS;
