#merge, replays: Cancel requests when demo stops and we didn't fully start recording + sanity checks

diff --git a/Engine/Source/Runtime/NetworkReplayStreaming/HttpNetworkReplayStreaming/Private/HttpNetworkReplayStreaming.cpp b/Engine/Source/Runtime/NetworkReplayStreaming/HttpNetworkReplayStreaming/Private/HttpNetworkReplayStreaming.cpp
--- a/Engine/Source/Runtime/NetworkReplayStreaming/HttpNetworkReplayStreaming/Private/HttpNetworkReplayStreaming.cpp
+++ b/Engine/Source/Runtime/NetworkReplayStreaming/HttpNetworkReplayStreaming/Private/HttpNetworkReplayStreaming.cpp
@@ -273,8 +273,17 @@ void FHttpNetworkReplayStreamer::AddRequestToQueue( const EQueuedHttpRequestType
 
 void FHttpNetworkReplayStreamer::StopStreaming()
 {
+	if ( StartStreamingDelegate.IsBound() )
+	{
+		UE_LOG( LogHttpReplay, Warning, TEXT( "FHttpNetworkReplayStreamer::StopStreaming. Called while existing StartStreaming request wasn't finished" ) );
+		CancelStreamingRequests();
+		check( !IsStreaming() );
+		return;
+	}
+
 	if ( !IsStreaming() )
 	{
+		UE_LOG( LogHttpReplay, Warning, TEXT( "FHttpNetworkReplayStreamer::StopStreaming. Not currently streaming." ) );
 		check( bStopStreamingCalled == false );
 		return;
 	}
@@ -300,6 +309,7 @@ void FHttpNetworkReplayStreamer::StopStreaming()
 		StopUploading();
 	}
 
+	// Finally, add the stop streaming request, which should put things in the right state after the above requests are done
 	AddRequestToQueue( EQueuedHttpRequestType::StopStreaming, nullptr );
 }
 
@@ -737,6 +747,13 @@ void FHttpNetworkReplayStreamer::ConditionallyRefreshViewer()
 };
 
 void FHttpNetworkReplayStreamer::SetLastError( const ENetworkReplayError::Type InLastError )
+{
+	CancelStreamingRequests();
+
+	StreamerLastError = InLastError;
+}
+
+void FHttpNetworkReplayStreamer::CancelStreamingRequests()
 {
 	// Cancel any in flight request
 	if ( InFlightHttpRequest.IsValid() )
@@ -747,10 +764,11 @@ void FHttpNetworkReplayStreamer::SetLastError( const ENetworkReplayError::Type I
 
 	// Empty the request queue
 	TSharedPtr< FQueuedHttpRequest > QueuedRequest;
-	while ( QueuedHttpRequests.Dequeue( QueuedRequest ) ) { }
+	while ( QueuedHttpRequests.Dequeue( QueuedRequest ) )
+	{
+	}
 
 	StreamerState			= EStreamerState::Idle;
-	StreamerLastError		= InLastError;
 	bStopStreamingCalled	= false;
 }
 
diff --git a/Engine/Source/Runtime/NetworkReplayStreaming/HttpNetworkReplayStreaming/Public/HttpNetworkReplayStreaming.h b/Engine/Source/Runtime/NetworkReplayStreaming/HttpNetworkReplayStreaming/Public/HttpNetworkReplayStreaming.h
--- a/Engine/Source/Runtime/NetworkReplayStreaming/HttpNetworkReplayStreaming/Public/HttpNetworkReplayStreaming.h
+++ b/Engine/Source/Runtime/NetworkReplayStreaming/HttpNetworkReplayStreaming/Public/HttpNetworkReplayStreaming.h
@@ -149,6 +149,7 @@ public:
 	void RefreshViewer( const bool bFinal );
 	void ConditionallyRefreshViewer();
 	void SetLastError( const ENetworkReplayError::Type InLastError );
+	void CancelStreamingRequests();
 	void FlushCheckpointInternal( uint32 TimeInMS );
 	void AddEvent( const uint32 TimeInMS, const FString& Group, const FString& Meta, const TArray<uint8>& Data );
 	void AddRequestToQueue( const EQueuedHttpRequestType::Type Type, TSharedPtr< class IHttpRequest >	Request );
