Refactored Replay API to be more like our REST standards.  Kept it as a v2 api so existing clients should be backwards compatible still.

TON-11903

+review CR [REDACTED_HANDLE]

diff --git a/Engine/Source/Runtime/NetworkReplayStreaming/HttpNetworkReplayStreaming/Private/HttpNetworkReplayStreaming.cpp b/Engine/Source/Runtime/NetworkReplayStreaming/HttpNetworkReplayStreaming/Private/HttpNetworkReplayStreaming.cpp
--- a/Engine/Source/Runtime/NetworkReplayStreaming/HttpNetworkReplayStreaming/Private/HttpNetworkReplayStreaming.cpp
+++ b/Engine/Source/Runtime/NetworkReplayStreaming/HttpNetworkReplayStreaming/Private/HttpNetworkReplayStreaming.cpp
@@ -69,6 +69,42 @@ public:
 	END_ONLINE_JSON_SERIALIZER
 };
 
+class FNetworkReplayStartUploadingResponse : public FOnlineJsonSerializable
+{
+public:
+	FNetworkReplayStartUploadingResponse()
+	{}
+	virtual ~FNetworkReplayStartUploadingResponse() {}
+
+	FString SessionId;
+
+	// FOnlineJsonSerializable
+	BEGIN_ONLINE_JSON_SERIALIZER
+		ONLINE_JSON_SERIALIZE("sessionId", SessionId);
+	END_ONLINE_JSON_SERIALIZER
+};
+
+class FNetworkReplayStartDownloadingResponse : public FOnlineJsonSerializable
+{
+public:
+	FNetworkReplayStartDownloadingResponse()
+	{}
+	virtual ~FNetworkReplayStartDownloadingResponse() {}
+
+	FString State;
+	FString Viewer;
+	int32 Time;
+	int32 NumChunks;
+
+	// FOnlineJsonSerializable
+	BEGIN_ONLINE_JSON_SERIALIZER
+		ONLINE_JSON_SERIALIZE("state", State);
+		ONLINE_JSON_SERIALIZE("numChunks", NumChunks);
+		ONLINE_JSON_SERIALIZE("time", Time);
+		ONLINE_JSON_SERIALIZE("viewerId", Viewer);
+	END_ONLINE_JSON_SERIALIZER
+};
+
 void FHttpStreamFArchive::Serialize( void* V, int64 Length ) 
 {
 	if ( IsLoading() )
@@ -206,7 +242,7 @@ void FHttpNetworkReplayStreamer::StartStreaming( const FString& CustomName, cons
 		}
 
 		// Notify the http server that we want to start downloading a replay
-		const FString URL = FString::Printf( TEXT( "%sstartdownloading?Session=%s&User=%s" ), *ServerURL, *SessionName, *UserName );
+		const FString URL = FString::Printf( TEXT( "%sreplay/%s/startDownloading?user=%s" ), *ServerURL, *SessionName, *UserName );
 
 		UE_LOG( LogHttpReplay, Verbose, TEXT( "FHttpNetworkReplayStreamer::StartStreaming. URL: %s" ), *URL );
 
@@ -236,12 +272,12 @@ void FHttpNetworkReplayStreamer::StartStreaming( const FString& CustomName, cons
 		if ( FParse::Value( FCommandLine::Get(), TEXT( "ReplayMeta=" ), MetaString ) && !MetaString.IsEmpty() )
 		{
 			// Notify the http server that we want to start uploading a replay
-			HttpRequest->SetURL( FString::Printf( TEXT( "%sstartuploading?App=%s&Version=%u&CL=%u&Friendly=%s&Meta=%s" ), *ServerURL, *ReplayVersion.AppString, ReplayVersion.NetworkVersion, ReplayVersion.Changelist, *FriendlyName, *MetaString ) );
+			HttpRequest->SetURL( FString::Printf( TEXT( "%sreplay?app=%s&version=%u&cl=%u&friendlyName=%s&meta=%s" ), *ServerURL, *ReplayVersion.AppString, ReplayVersion.NetworkVersion, ReplayVersion.Changelist, *FriendlyName, *MetaString ) );
 		}
 		else
 		{
 			// Notify the http server that we want to start uploading a replay
-			HttpRequest->SetURL( FString::Printf( TEXT( "%sstartuploading?App=%s&Version=%u&CL=%u&Friendly=%s" ), *ServerURL, *ReplayVersion.AppString, ReplayVersion.NetworkVersion, ReplayVersion.Changelist, *FriendlyName ) );
+			HttpRequest->SetURL( FString::Printf( TEXT( "%sreplay?app=%s&version=%u&cl=%u&friendlyName=%s" ), *ServerURL, *ReplayVersion.AppString, ReplayVersion.NetworkVersion, ReplayVersion.Changelist, *FriendlyName ) );
 		}
 
 		HttpRequest->SetVerb( TEXT( "POST" ) );
@@ -341,7 +377,7 @@ void FHttpNetworkReplayStreamer::UploadHeader()
 
 	HttpRequest->OnProcessRequestComplete().BindRaw( this, &FHttpNetworkReplayStreamer::HttpHeaderUploadFinished );
 
-	HttpRequest->SetURL( FString::Printf( TEXT( "%supload?Session=%s&NumChunks=%i&Time=%i&Filename=replay.header" ), *ServerURL, *SessionName, StreamChunkIndex, TotalDemoTimeInMS ) );
+	HttpRequest->SetURL( FString::Printf( TEXT( "%sreplay/%s/file/replay.header?numChunks=%i&time=%i" ), *ServerURL, *SessionName, StreamChunkIndex, TotalDemoTimeInMS ) );
 	HttpRequest->SetVerb( TEXT( "POST" ) );
 	HttpRequest->SetHeader( TEXT( "Content-Type" ), TEXT( "application/octet-stream" ) );
 	HttpRequest->SetContent( HeaderArchive.Buffer );
@@ -380,7 +416,7 @@ void FHttpNetworkReplayStreamer::FlushStream()
 
 	HttpRequest->OnProcessRequestComplete().BindRaw( this, &FHttpNetworkReplayStreamer::HttpUploadStreamFinished );
 
-	HttpRequest->SetURL( FString::Printf( TEXT( "%supload?Session=%s&NumChunks=%i&Time=%i&MTime1=%i&MTime2=%i&Filename=stream.%i" ), *ServerURL, *SessionName, StreamChunkIndex + 1, TotalDemoTimeInMS, StreamTimeRangeStart, StreamTimeRangeEnd, StreamChunkIndex ) );
+	HttpRequest->SetURL(FString::Printf(TEXT("%sreplay/%s/file/stream.%i?numChunks=%i&time=%i&mTime1=%i&mTime2=%i"), *ServerURL, *SessionName, StreamChunkIndex, StreamChunkIndex + 1, TotalDemoTimeInMS, StreamTimeRangeStart, StreamTimeRangeEnd));
 	HttpRequest->SetVerb( TEXT( "POST" ) );
 	HttpRequest->SetHeader( TEXT( "Content-Type" ), TEXT( "application/octet-stream" ) );
 	HttpRequest->SetContent( StreamArchive.Buffer );
@@ -415,7 +451,7 @@ void FHttpNetworkReplayStreamer::StopUploading()
 
 	HttpRequest->OnProcessRequestComplete().BindRaw( this, &FHttpNetworkReplayStreamer::HttpStopUploadingFinished );
 
-	HttpRequest->SetURL( FString::Printf( TEXT( "%sstopuploading?Session=%s&NumChunks=%i&Time=%i" ), *ServerURL, *SessionName, StreamChunkIndex, TotalDemoTimeInMS ) );
+	HttpRequest->SetURL( FString::Printf( TEXT( "%sreplay/%s/stopUploading?numChunks=%i&time=%i" ), *ServerURL, *SessionName, StreamChunkIndex, TotalDemoTimeInMS ) );
 	HttpRequest->SetVerb( TEXT( "POST" ) );
 	HttpRequest->SetHeader( TEXT( "Content-Type" ), TEXT( "application/octet-stream" ) );
 
@@ -481,7 +517,7 @@ void FHttpNetworkReplayStreamer::GotoCheckpointIndex( const int32 CheckpointInde
 	TSharedRef<class IHttpRequest> HttpRequest = FHttpModule::Get().CreateRequest();
 
 	// Download the next stream chunk
-	HttpRequest->SetURL( FString::Printf( TEXT( "%sdownloadevent?Session=%s&ID=%s" ), *ServerURL, *SessionName, *CheckpointList.ReplayEvents[CheckpointIndex].ID ) );
+	HttpRequest->SetURL( FString::Printf( TEXT( "%sevent/%s" ), *ServerURL,  *CheckpointList.ReplayEvents[CheckpointIndex].ID ) );
 	HttpRequest->SetVerb( TEXT( "GET" ) );
 
 	HttpRequest->OnProcessRequestComplete().BindRaw( this, &FHttpNetworkReplayStreamer::HttpDownloadCheckpointFinished );
@@ -496,7 +532,7 @@ void FHttpNetworkReplayStreamer::SearchEvents(const FString& EventGroup, const F
 {
 	TSharedRef<class IHttpRequest> HttpRequest = FHttpModule::Get().CreateRequest();
 
-	HttpRequest->SetURL(FString::Printf(TEXT("%ssearchevents?group=%s"), *ServerURL, *EventGroup));
+	HttpRequest->SetURL(FString::Printf(TEXT("%sevent?group=%s"), *ServerURL, *EventGroup));
 	HttpRequest->SetVerb(TEXT("GET"));
 
 	HttpRequest->OnProcessRequestComplete().BindRaw(this, &FHttpNetworkReplayStreamer::HttpEnumerateSessionsFinished);
@@ -511,7 +547,7 @@ void FHttpNetworkReplayStreamer::RequestEventData(const FString& EventID, const
 	TSharedRef<class IHttpRequest> HttpRequest = FHttpModule::Get().CreateRequest();
 
 	// Download the next stream chunk
-	HttpRequest->SetURL(FString::Printf(TEXT("%sdownloadevent?Session=%s&ID=%s"), *ServerURL, *SessionName, *EventID));
+	HttpRequest->SetURL(FString::Printf(TEXT("%sevent/%s"), *ServerURL, *EventID));
 	HttpRequest->SetVerb(TEXT("GET"));
 
 	HttpRequest->OnProcessRequestComplete().BindRaw(this, &FHttpNetworkReplayStreamer::HttpRequestEventDataFinished, Delegate);
@@ -608,7 +644,7 @@ void FHttpNetworkReplayStreamer::FlushCheckpointInternal( uint32 TimeInMS )
 
 	HttpRequest->OnProcessRequestComplete().BindRaw( this, &FHttpNetworkReplayStreamer::HttpUploadCheckpointFinished );
 
-	HttpRequest->SetURL( FString::Printf( TEXT( "%suploadevent?Session=%s&Group=checkpoint&Time1=%i&Time2=%i&Meta=%i" ), *ServerURL, *SessionName, TimeInMS, TimeInMS, StreamChunkIndex ) );
+	HttpRequest->SetURL( FString::Printf( TEXT( "%sreplay/%s/event?group=checkpoint&time1=%i&time2=%i&meta=%i" ), *ServerURL, *SessionName, TimeInMS, TimeInMS, StreamChunkIndex ) );
 	HttpRequest->SetVerb( TEXT( "POST" ) );
 	HttpRequest->SetHeader( TEXT( "Content-Type" ), TEXT( "application/octet-stream" ) );
 	HttpRequest->SetContent( CheckpointArchive.Buffer );
@@ -634,7 +670,7 @@ void FHttpNetworkReplayStreamer::AddEvent( const uint32 TimeInMS, const FString&
 
 	HttpRequest->OnProcessRequestComplete().BindRaw(this, &FHttpNetworkReplayStreamer::HttpUploadCustomEventFinished);
 
-	HttpRequest->SetURL(FString::Printf(TEXT("%suploadevent?Session=%s&Group=%s&Time1=%i&Time2=%i&Meta=%s"), *ServerURL, *SessionName, *Group, TimeInMS, TimeInMS, *Meta));
+	HttpRequest->SetURL(FString::Printf(TEXT("%sreplay/%s/event?group=%s&time1=%i&time2=%i&meta=%s"), *ServerURL, *SessionName, *Group, TimeInMS, TimeInMS, *Meta));
 	HttpRequest->SetVerb(TEXT("POST"));
 	HttpRequest->SetHeader(TEXT("Content-Type"), TEXT("application/octet-stream"));
 	HttpRequest->SetContent(Data);
@@ -647,7 +683,7 @@ void FHttpNetworkReplayStreamer::DownloadHeader()
 	// Download header first
 	TSharedRef<class IHttpRequest> HttpRequest = FHttpModule::Get().CreateRequest();
 
-	HttpRequest->SetURL( FString::Printf( TEXT( "%sdownload?Session=%s&Filename=replay.header" ), *ServerURL, *SessionName, *ViewerName ) );
+	HttpRequest->SetURL( FString::Printf( TEXT( "%sreplay/%s/file/replay.header" ), *ServerURL, *SessionName, *ViewerName ) );
 	HttpRequest->SetVerb( TEXT( "GET" ) );
 
 	HttpRequest->OnProcessRequestComplete().BindRaw( this, &FHttpNetworkReplayStreamer::HttpDownloadHeaderFinished );
@@ -699,7 +735,7 @@ void FHttpNetworkReplayStreamer::ConditionallyDownloadNextChunk()
 	TSharedRef<class IHttpRequest> HttpRequest = FHttpModule::Get().CreateRequest();
 
 	// Download the next stream chunk
-	const FString URL = FString::Printf( TEXT( "%sdownload?Session=%s&Filename=stream.%i" ), *ServerURL, *SessionName, StreamChunkIndex );
+	const FString URL = FString::Printf( TEXT( "%sreplay/%s/file/stream.%i" ), *ServerURL, *SessionName, StreamChunkIndex );
 
 	UE_LOG( LogHttpReplay, Verbose, TEXT( "FHttpNetworkReplayStreamer::ConditionallyDownloadNextChunk. URL: %s" ), *URL );
 
@@ -720,11 +756,11 @@ void FHttpNetworkReplayStreamer::RefreshViewer( const bool bFinal )
 	// Download the next stream chunk
 	if ( bFinal )
 	{
-		HttpRequest->SetURL( FString::Printf( TEXT( "%srefreshviewer?Session=%s&Viewer=%s&Final=true" ), *ServerURL, *SessionName, *ViewerName ) );
+		HttpRequest->SetURL( FString::Printf( TEXT( "%sreplay/%s/viewer/%s?final=true" ), *ServerURL, *SessionName, *ViewerName ) );
 	}
 	else
 	{
-		HttpRequest->SetURL( FString::Printf( TEXT( "%srefreshviewer?Session=%s&Viewer=%s" ), *ServerURL, *SessionName, *ViewerName ) );
+		HttpRequest->SetURL( FString::Printf( TEXT( "%sreplay/%s/viewer/%s" ), *ServerURL, *SessionName, *ViewerName ) );
 	}
 
 	HttpRequest->SetVerb( TEXT( "POST" ) );
@@ -884,17 +920,17 @@ void FHttpNetworkReplayStreamer::EnumerateStreams( const FNetworkReplayVersion&
 	TSharedRef<class IHttpRequest> HttpRequest = FHttpModule::Get().CreateRequest();
 
 	// Build base URL
-	FString URL = FString::Printf( TEXT( "%senumsessions?App=%s" ), *ServerURL, *InReplayVersion.AppString );
+	FString URL = FString::Printf( TEXT( "%sreplay?app=%s" ), *ServerURL, *InReplayVersion.AppString );
 
 	// Add optional stuff
 	if ( InReplayVersion.Changelist != 0 )
 	{
-		URL += FString::Printf( TEXT( "&CL=%u" ), InReplayVersion.Changelist );
+		URL += FString::Printf( TEXT( "&cl=%u" ), InReplayVersion.Changelist );
 	}
 
 	if ( InReplayVersion.NetworkVersion != 0 )
 	{
-		URL += FString::Printf( TEXT( "&Version=%u" ), InReplayVersion.NetworkVersion );
+		URL += FString::Printf( TEXT( "&version=%u" ), InReplayVersion.NetworkVersion );
 	}
 
 	const FString MetaStringToUse = !CVarMetaFilterOverride.GetValueOnGameThread().IsEmpty() ? CVarMetaFilterOverride.GetValueOnGameThread() : MetaString;
@@ -902,13 +938,13 @@ void FHttpNetworkReplayStreamer::EnumerateStreams( const FNetworkReplayVersion&
 	// Add optional Meta parameter (filter replays by meta tag)
 	if ( !MetaStringToUse.IsEmpty() )
 	{
-		URL += FString::Printf( TEXT( "&Meta=%s" ), *MetaStringToUse );
+		URL += FString::Printf( TEXT( "&meta=%s" ), *MetaStringToUse );
 	}
 
 	// Add optional User parameter (filter replays by a user that was in the replay)
 	if ( !UserString.IsEmpty() )
 	{
-		URL += FString::Printf( TEXT( "&User=%s" ), *UserString );
+		URL += FString::Printf( TEXT( "&user=%s" ), *UserString );
 	}
 
 	HttpRequest->SetURL( URL );
@@ -933,7 +969,7 @@ void FHttpNetworkReplayStreamer::EnumerateRecentStreams( const FNetworkReplayVer
 	TSharedRef<class IHttpRequest> HttpRequest = FHttpModule::Get().CreateRequest();
 
 	// Enumerate all of the sessions
-	HttpRequest->SetURL( FString::Printf( TEXT( "%senumsessions?App=%s&Version=%u&CL=%u&Recent=%s" ), *ServerURL, *InReplayVersion.AppString, InReplayVersion.NetworkVersion, InReplayVersion.Changelist, *InRecentViewer ) );
+	HttpRequest->SetURL( FString::Printf( TEXT( "%sreplay?app=%s&version=%u&cl=%u&recent=%s" ), *ServerURL, *InReplayVersion.AppString, InReplayVersion.NetworkVersion, InReplayVersion.Changelist, *InRecentViewer ) );
 	HttpRequest->SetVerb( TEXT( "GET" ) );
 
 	HttpRequest->OnProcessRequestComplete().BindRaw( this, &FHttpNetworkReplayStreamer::HttpEnumerateSessionsFinished );
@@ -958,7 +994,7 @@ void FHttpNetworkReplayStreamer::AddUserToReplay(const FString& UserString)
 
 	TSharedRef<class IHttpRequest> HttpRequest = FHttpModule::Get().CreateRequest();
 
-	HttpRequest->SetURL( FString::Printf( TEXT( "%saddusers?Session=%s" ), *ServerURL, *SessionName ) );
+	HttpRequest->SetURL( FString::Printf( TEXT( "%sreplay/%s/users" ), *ServerURL, *SessionName ) );
 	HttpRequest->SetVerb( TEXT( "POST" ) );
 
 	FNetworkReplayUserList UserList;
@@ -978,7 +1014,7 @@ void FHttpNetworkReplayStreamer::EnumerateCheckpoints()
 	TSharedRef<class IHttpRequest> HttpRequest = FHttpModule::Get().CreateRequest();
 
 	// Enumerate all of the sessions
-	HttpRequest->SetURL( FString::Printf( TEXT( "%senumevents?Session=%s&Group=checkpoint" ), *ServerURL, *SessionName ) );
+	HttpRequest->SetURL( FString::Printf( TEXT( "%sreplay/%s/event?group=checkpoint" ), *ServerURL, *SessionName ) );
 	HttpRequest->SetVerb( TEXT( "GET" ) );
 
 	HttpRequest->OnProcessRequestComplete().BindRaw( this, &FHttpNetworkReplayStreamer::HttpEnumerateCheckpointsFinished );
@@ -1011,7 +1047,7 @@ void FHttpNetworkReplayStreamer::EnumerateEvents(const FString& Group, const FEn
 	TSharedRef<class IHttpRequest> HttpRequest = FHttpModule::Get().CreateRequest();
 
 	// Enumerate all of the sessions
-	HttpRequest->SetURL(FString::Printf(TEXT("%senumevents?Session=%s&Group=%s"), *ServerURL, *SessionName, *Group));
+	HttpRequest->SetURL(FString::Printf(TEXT("%sreplay/%s/event?group=%s"), *ServerURL, *SessionName, *Group));
 	HttpRequest->SetVerb(TEXT("GET"));
 
 	HttpRequest->OnProcessRequestComplete().BindRaw(this, &FHttpNetworkReplayStreamer::HttpEnumerateEventsFinished, EnumerationCompleteDelegate);
@@ -1035,7 +1071,18 @@ void FHttpNetworkReplayStreamer::HttpStartUploadingFinished( FHttpRequestPtr Htt
 
 	if ( bSucceeded && HttpResponse->GetResponseCode() == EHttpResponseCodes::Ok )
 	{
-		SessionName = HttpResponse->GetHeader( TEXT( "Session" ) );
+		FString JsonString = HttpResponse->GetContentAsString();
+
+		FNetworkReplayStartUploadingResponse StartUploadingResponse;
+
+		if (!StartUploadingResponse.FromJson(JsonString))
+		{
+			UE_LOG(LogHttpReplay, Warning, TEXT("FHttpNetworkReplayStreamer::HttpStartUploadingFinished. FromJson FAILED"));
+			// FIXME: Is there more I should do here?
+			return;
+		}
+
+		SessionName = StartUploadingResponse.SessionId;
 
 		UE_LOG( LogHttpReplay, Log, TEXT( "FHttpNetworkReplayStreamer::HttpStartUploadingFinished. SessionName: %s" ), *SessionName );
 	}
@@ -1050,7 +1097,7 @@ void FHttpNetworkReplayStreamer::HttpStopUploadingFinished( FHttpRequestPtr Http
 {
 	RequestFinished( EStreamerState::StreamingUp, EQueuedHttpRequestType::StopUploading, HttpRequest );
 
-	if ( bSucceeded && HttpResponse->GetResponseCode() == EHttpResponseCodes::Ok )
+	if ( bSucceeded && HttpResponse->GetResponseCode() == EHttpResponseCodes::NoContent )
 	{
 		UE_LOG( LogHttpReplay, Log, TEXT( "FHttpNetworkReplayStreamer::HttpStopUploadingFinished. SessionName: %s" ), *SessionName );
 	}
@@ -1074,7 +1121,7 @@ void FHttpNetworkReplayStreamer::HttpHeaderUploadFinished( FHttpRequestPtr HttpR
 
 	check( StartStreamingDelegate.IsBound() );
 
-	if ( bSucceeded && HttpResponse->GetResponseCode() == EHttpResponseCodes::Ok )
+	if ( bSucceeded && HttpResponse->GetResponseCode() == EHttpResponseCodes::NoContent )
 	{
 		UE_LOG( LogHttpReplay, Log, TEXT( "FHttpNetworkReplayStreamer::HttpHeaderUploadFinished." ) );
 
@@ -1095,7 +1142,7 @@ void FHttpNetworkReplayStreamer::HttpUploadStreamFinished( FHttpRequestPtr HttpR
 {
 	RequestFinished( EStreamerState::StreamingUp, EQueuedHttpRequestType::UploadingStream, HttpRequest );
 
-	if ( bSucceeded && HttpResponse->GetResponseCode() == EHttpResponseCodes::Ok )
+	if ( bSucceeded && HttpResponse->GetResponseCode() == EHttpResponseCodes::NoContent )
 	{
 		UE_LOG( LogHttpReplay, Verbose, TEXT( "FHttpNetworkReplayStreamer::HttpUploadStreamFinished." ) );
 	}
@@ -1110,7 +1157,7 @@ void FHttpNetworkReplayStreamer::HttpUploadCheckpointFinished( FHttpRequestPtr H
 {
 	RequestFinished( EStreamerState::StreamingUp, EQueuedHttpRequestType::UploadingCheckpoint, HttpRequest );
 
-	if ( bSucceeded && HttpResponse->GetResponseCode() == EHttpResponseCodes::Ok )
+	if ( bSucceeded && HttpResponse->GetResponseCode() == EHttpResponseCodes::NoContent )
 	{
 		UE_LOG( LogHttpReplay, Verbose, TEXT( "FHttpNetworkReplayStreamer::HttpUploadCheckpointFinished." ) );
 	}
@@ -1147,16 +1194,25 @@ void FHttpNetworkReplayStreamer::HttpStartDownloadingFinished( FHttpRequestPtr H
 
 	if ( bSucceeded && HttpResponse->GetResponseCode() == EHttpResponseCodes::Ok )
 	{
-		FString NumChunksString = HttpResponse->GetHeader( TEXT( "NumChunks" ) );
-		FString DemoTimeString = HttpResponse->GetHeader( TEXT( "Time" ) );
-		FString State = HttpResponse->GetHeader( TEXT( "State" ) );
 
-		ViewerName = HttpResponse->GetHeader( TEXT( "Viewer" ) );
+		FString JsonString = HttpResponse->GetContentAsString();
+
+		FNetworkReplayStartDownloadingResponse StartDownloadingResponse;
+
+		if (!StartDownloadingResponse.FromJson(JsonString))
+		{
+			UE_LOG(LogHttpReplay, Warning, TEXT("FHttpNetworkReplayStreamer::HttpStartDownloadingFinished. FromJson FAILED"));
+			// FIXME: Is there more I should do here?
+			return;
+		}
+
+		FString State = StartDownloadingResponse.State;
+		ViewerName = StartDownloadingResponse.Viewer;
 
 		bStreamIsLive = State == TEXT( "Live" );
 
-		NumTotalStreamChunks	= FCString::Atoi( *NumChunksString );
-		TotalDemoTimeInMS		= FCString::Atoi( *DemoTimeString );
+		NumTotalStreamChunks = StartDownloadingResponse.NumChunks;
+		TotalDemoTimeInMS = StartDownloadingResponse.Time;
 
 		UE_LOG( LogHttpReplay, Log, TEXT( "FHttpNetworkReplayStreamer::HttpStartDownloadingFinished. Viewer: %s, State: %s, NumChunks: %i, DemoTime: %2.2f" ), *ViewerName, *State, NumTotalStreamChunks, (float)TotalDemoTimeInMS / 1000 );
 
@@ -1355,7 +1411,7 @@ void FHttpNetworkReplayStreamer::HttpRefreshViewerFinished( FHttpRequestPtr Http
 {
 	RequestFinished( EStreamerState::StreamingDown, EQueuedHttpRequestType::RefreshingViewer, HttpRequest );
 
-	if ( !bSucceeded || HttpResponse->GetResponseCode() != EHttpResponseCodes::Ok )
+	if ( !bSucceeded || HttpResponse->GetResponseCode() != EHttpResponseCodes::NoContent )
 	{
 		UE_LOG( LogHttpReplay, Error, TEXT( "FHttpNetworkReplayStreamer::HttpRefreshViewerFinished. FAILED, Response code: %d" ), HttpResponse.IsValid() ? HttpResponse->GetResponseCode() : 0 );
 		SetLastError( ENetworkReplayError::ServiceUnavailable );
@@ -1504,7 +1560,7 @@ void FHttpNetworkReplayStreamer::HttpAddUserFinished(FHttpRequestPtr HttpRequest
 {
 	RequestFinished( EStreamerState::StreamingUp, EQueuedHttpRequestType::AddingUser, HttpRequest );
 
-	if ( bSucceeded && HttpResponse->GetResponseCode() == EHttpResponseCodes::Ok )
+	if ( bSucceeded && HttpResponse->GetResponseCode() == EHttpResponseCodes::NoContent )
 	{		
 		UE_LOG( LogHttpReplay, Verbose, TEXT( "FHttpNetworkReplayStreamer::HttpAddUserFinished." ) );
 	}
