FNullNetworkReplayStreamer can now play back live replays.

--------
Integrated using branch UE4-Orion-To-UE4 of change#2677431 by [REDACTED_PERSON] on 2015/09/02 15:08:36.

diff --git a/Engine/Source/Runtime/Engine/Private/DemoNetDriver.cpp b/Engine/Source/Runtime/Engine/Private/DemoNetDriver.cpp
--- a/Engine/Source/Runtime/Engine/Private/DemoNetDriver.cpp
+++ b/Engine/Source/Runtime/Engine/Private/DemoNetDriver.cpp
@@ -711,6 +711,7 @@ bool UDemoNetDriver::InitListen( FNetworkNotify* InNotify, FURL& ListenURL, bool
 
 	// Write the header
 	(*FileAr) << DemoHeader;
+	FileAr->Flush();
 
 	// Spawn the demo recording spectator.
 	SpawnDemoRecSpectator( Connection );
diff --git a/Engine/Source/Runtime/NetworkReplayStreaming/NullNetworkReplayStreaming/NullNetworkReplayStreaming.build.cs b/Engine/Source/Runtime/NetworkReplayStreaming/NullNetworkReplayStreaming/NullNetworkReplayStreaming.build.cs
--- a/Engine/Source/Runtime/NetworkReplayStreaming/NullNetworkReplayStreaming/NullNetworkReplayStreaming.build.cs
+++ b/Engine/Source/Runtime/NetworkReplayStreaming/NullNetworkReplayStreaming/NullNetworkReplayStreaming.build.cs
@@ -10,12 +10,15 @@ namespace UnrealBuildTool.Rules
 
             PrivateIncludePathModuleNames.Add("OnlineSubsystem");
 
+			PublicIncludePathModuleNames.Add("Engine");
+
 			PrivateDependencyModuleNames.AddRange(
 				new string[]
 				{
 					"Core",
 					"NetworkReplayStreaming",
-                    "Json"
+                    "Json",
+					"Engine"
 				}
 			);
 		}
diff --git a/Engine/Source/Runtime/NetworkReplayStreaming/NullNetworkReplayStreaming/Private/NullNetworkReplayStreaming.cpp b/Engine/Source/Runtime/NetworkReplayStreaming/NullNetworkReplayStreaming/Private/NullNetworkReplayStreaming.cpp
--- a/Engine/Source/Runtime/NetworkReplayStreaming/NullNetworkReplayStreaming/Private/NullNetworkReplayStreaming.cpp
+++ b/Engine/Source/Runtime/NetworkReplayStreaming/NullNetworkReplayStreaming/Private/NullNetworkReplayStreaming.cpp
@@ -112,17 +112,30 @@ static FNullReplayInfo ReadReplayInfo( const FString& StreamName )
 	const FString InfoFilename = GetInfoFilename(StreamName);
 	TUniquePtr<FArchive> InfoFileArchive( IFileManager::Get().CreateFileReader( *InfoFilename ) );
 
-	if ( InfoFileArchive.IsValid() )
+	if ( InfoFileArchive.IsValid() && InfoFileArchive->TotalSize() != 0)
 	{
 		FString JsonString;
 		*InfoFileArchive << JsonString;
 
 		Info.FromJson(JsonString);
+		Info.bIsValid = true;
 	}
 
 	return Info;
 }
 
+static void WriteReplayInfo( const FString& StreamName, const FNullReplayInfo& ReplayInfo )
+{
+	// Update metadata file with latest info
+	TUniquePtr<FArchive> ReplayInfoFileAr(IFileManager::Get().CreateFileWriter(*GetInfoFilename(StreamName)));
+
+	if (ReplayInfoFileAr.IsValid())
+	{
+		FString JsonString = ReplayInfo.ToJson();
+		*ReplayInfoFileAr << JsonString; 
+	}
+}
+
 // Returns a name formatted as "demoX", where X is 0-9.
 // Returns the first value that doesn't yet exist, or if they all exist, returns the oldest one
 // (it will be overwritten).
@@ -190,8 +203,8 @@ void FNullNetworkReplayStreamer::StartStreaming( const FString& CustomName, cons
 		ReplayInfo = ReadReplayInfo( CurrentStreamName );
 
 		// Open file for reading
-		FileAr.Reset( IFileManager::Get().CreateFileReader( *FullDemoFilename ) );
-		HeaderAr.Reset( IFileManager::Get().CreateFileReader( *FullHeaderFilename ) );
+		ReopenStreamFileForReading();
+		HeaderAr.Reset( IFileManager::Get().CreateFileReader( *FullHeaderFilename, FILEREAD_AllowWrite ) );
 		StreamerState = EStreamerState::Playback;
 	}
 	else
@@ -204,7 +217,7 @@ void FNullNetworkReplayStreamer::StartStreaming( const FString& CustomName, cons
 
 		// Open file for writing
 		FileAr.Reset( IFileManager::Get().CreateFileWriter( *FullDemoFilename, FILEWRITE_AllowRead ) );
-		HeaderAr.Reset( IFileManager::Get().CreateFileWriter( *FullHeaderFilename ) );
+		HeaderAr.Reset( IFileManager::Get().CreateFileWriter( *FullHeaderFilename, FILEWRITE_AllowRead ) );
 		StreamerState = EStreamerState::Recording;
 
 		CurrentCheckpointIndex = 0;
@@ -213,6 +226,8 @@ void FNullNetworkReplayStreamer::StartStreaming( const FString& CustomName, cons
 		ReplayInfo.NetworkVersion = ReplayVersion.NetworkVersion;
 		ReplayInfo.Changelist = ReplayVersion.Changelist;
 		ReplayInfo.FriendlyName = FriendlyName;
+
+		WriteReplayInfo(CurrentStreamName, ReplayInfo);
 	}
 
 	// Notify immediately
@@ -223,14 +238,7 @@ void FNullNetworkReplayStreamer::StopStreaming()
 {
 	if (StreamerState == EStreamerState::Recording)
 	{
-		// Update metadata file with latest info
-		TUniquePtr<FArchive> ReplayInfoFileAr(IFileManager::Get().CreateFileWriter(*GetInfoFilename(CurrentStreamName)));
-
-		if (ReplayInfoFileAr.IsValid())
-		{
-			FString JsonString = ReplayInfo.ToJson();
-			*ReplayInfoFileAr << JsonString; 
-		}
+		WriteReplayInfo(CurrentStreamName, ReplayInfo);
 	}
 
 	HeaderAr.Reset();
@@ -278,9 +286,18 @@ FArchive* FNullNetworkReplayStreamer::GetMetadataArchive()
 
 void FNullNetworkReplayStreamer::UpdateTotalDemoTime(uint32 TimeInMS)
 {
+	check(StreamerState == EStreamerState::Recording);
+
 	ReplayInfo.LengthInMS = TimeInMS;
 }
 
+bool FNullNetworkReplayStreamer::IsDataAvailable() const
+{
+	check(StreamerState == EStreamerState::Playback);
+
+	return FileAr.IsValid() && FileAr->Tell() < ReplayInfo.DemoFileLastOffset;
+}
+
 bool FNullNetworkReplayStreamer::IsLive() const
 {
 	return IsNamedStreamLive(CurrentStreamName);
@@ -288,15 +305,8 @@ bool FNullNetworkReplayStreamer::IsLive() const
 
 bool FNullNetworkReplayStreamer::IsNamedStreamLive( const FString& StreamName ) const
 {
-	// If the directory for this stream doesn't exist, it can't possibly be live.
-	if (!IFileManager::Get().DirectoryExists(*GetStreamDirectory(StreamName)))
-	{
-		return false;
-	}
-
 	// If the metadata file doesn't exist, this is a live stream.
-	const int64 MetadataFileSize = IFileManager::Get().FileSize(*GetMetadataFilename(StreamName));
-	return MetadataFileSize == INDEX_NONE;
+	return !IFileManager::Get().FileExists(*GetMetadataFilename(StreamName));
 }
 
 void FNullNetworkReplayStreamer::DeleteFinishedStream( const FString& StreamName, const FOnDeleteFinishedStreamComplete& Delegate ) const
@@ -335,32 +345,30 @@ void FNullNetworkReplayStreamer::EnumerateStreams( const FNetworkReplayVersion&
 		FNetworkReplayStreamInfo Info;
 		Info.SizeInBytes = IFileManager::Get().FileSize( *FullDemoFilePath );
 		
-		if (Info.SizeInBytes != INDEX_NONE)
+		// Read stored info for this replay
+		FNullReplayInfo StoredReplayInfo = ReadReplayInfo( Directory );
+		
+		if (!StoredReplayInfo.bIsValid)
 		{
-			// Read stored info for this replay
-			FNullReplayInfo StoredReplayInfo = ReadReplayInfo( Directory );
-			
-			// Check version. NetworkVersion and changelist of 0 will ignore version check.
-			const bool NetworkVersionMatches = ReplayVersion.NetworkVersion == StoredReplayInfo.NetworkVersion;
-			const bool ChangelistMatches = ReplayVersion.Changelist == StoredReplayInfo.Changelist;
+			continue;
+		}
 
-			const bool NetworkVersionPasses = ReplayVersion.NetworkVersion == 0 || NetworkVersionMatches;
-			const bool ChangelistPasses = ReplayVersion.Changelist == 0 || ChangelistMatches;
+		// Check version. NetworkVersion and changelist of 0 will ignore version check.
+		const bool NetworkVersionMatches = ReplayVersion.NetworkVersion == StoredReplayInfo.NetworkVersion;
+		const bool ChangelistMatches = ReplayVersion.Changelist == StoredReplayInfo.Changelist;
 
-			if ( NetworkVersionPasses && ChangelistPasses )
-			{
-				Info.Name = Directory;
-				Info.Timestamp = IFileManager::Get().GetTimeStamp( *FullDemoFilePath );
-				Info.bIsLive = IsNamedStreamLive( Directory );
-				Info.LengthInMS = StoredReplayInfo.LengthInMS;
-				Info.FriendlyName = StoredReplayInfo.FriendlyName;
-
-				// Live streams not supported yet
-				if (!Info.bIsLive)
-				{
-					Results.Add(Info);
-				}
-			}
+		const bool NetworkVersionPasses = ReplayVersion.NetworkVersion == 0 || NetworkVersionMatches;
+		const bool ChangelistPasses = ReplayVersion.Changelist == 0 || ChangelistMatches;
+
+		if ( NetworkVersionPasses && ChangelistPasses )
+		{
+			Info.Name = Directory;
+			Info.Timestamp = IFileManager::Get().GetTimeStamp( *FullDemoFilePath );
+			Info.bIsLive = IsNamedStreamLive( Directory );
+			Info.LengthInMS = StoredReplayInfo.LengthInMS;
+			Info.FriendlyName = StoredReplayInfo.FriendlyName;
+
+			Results.Add(Info);
 		}
 	}
 
@@ -482,12 +490,36 @@ void FNullNetworkReplayStreamer::GotoCheckpointIndexInternal(int32 CheckpointInd
 		FNullCheckpointListItem Item;
 		Item.FromJson(JsonString);
 
+		// Reopen, since for live replays the file is being written to and read from simultaneously
+		// and we need the reported file size to be up to date.
+		ReopenStreamFileForReading();
+
 		FileAr->Seek( FCString::Atoi64( *Item.Metadata ) );
 	}
 
 	Delegate.ExecuteIfBound( true, TimeInMS );
 }
 
+void FNullNetworkReplayStreamer::ReopenStreamFileForReading()
+{
+	const FString FullName = GetDemoFilename(CurrentStreamName);
+	FileAr.Reset( IFileManager::Get().CreateFileReader( *FullName, FILEREAD_AllowWrite ) );
+	if (FileAr.IsValid())
+	{
+		LastKnownFileSize = FileAr->TotalSize();
+	}
+}
+
+void FNullNetworkReplayStreamer::UpdateReplayInfoIfValid()
+{
+	FNullReplayInfo LatestInfo = ReadReplayInfo(CurrentStreamName);
+
+	if (LatestInfo.bIsValid)
+	{
+		ReplayInfo = LatestInfo;
+	}
+}
+
 void FNullNetworkReplayStreamer::GotoTimeInMS(const uint32 TimeInMS, const FOnCheckpointReadyDelegate& Delegate)
 {
 	// Enumerate all the events in the events folder, since we need to know what times the checkpoints correlate with
@@ -555,6 +587,44 @@ void FNullNetworkReplayStreamer::GotoTimeInMS(const uint32 TimeInMS, const FOnCh
 	GotoCheckpointIndexInternal( CheckpointIndex, Delegate, ExtraSkipTimeInMS );
 }
 
+void FNullNetworkReplayStreamer::Tick(float DeltaSeconds)
+{
+	// This relies on the fact that the DemoNetDriver isn't currently in the middle of its own tick,
+	// and has either read or written a whole demo frame.
+	if (StreamerState == EStreamerState::Playback)
+	{
+		// Re-read replay info
+		UpdateReplayInfoIfValid();
+
+		// If there are new whole frames to read in the file, reopen it to refresh the size
+		if (ReplayInfo.DemoFileLastOffset > LastKnownFileSize)
+		{
+			const int64 OldLocation = FileAr->Tell();
+			ReopenStreamFileForReading();
+			if (FileAr.IsValid())
+			{
+				FileAr->Seek(OldLocation);
+			}
+		}
+	}
+	else if (StreamerState == EStreamerState::Recording)
+	{
+		// Note the size of the file between demo frames
+		if (FileAr.IsValid() && ReplayInfo.DemoFileLastOffset < FileAr->Tell())
+		{
+			ReplayInfo.DemoFileLastOffset = FileAr->Tell();
+			FileAr->Flush();
+
+			WriteReplayInfo(CurrentStreamName, ReplayInfo);
+		}
+	}
+}
+
+TStatId FNullNetworkReplayStreamer::GetStatId() const
+{
+	RETURN_QUICK_DECLARE_CYCLE_STAT(FNullNetworkReplayStreamer, STATGROUP_Tickables);
+}
+
 IMPLEMENT_MODULE( FNullNetworkReplayStreamingFactory, NullNetworkReplayStreaming )
 
 TSharedPtr< INetworkReplayStreamer > FNullNetworkReplayStreamingFactory::CreateReplayStreamer() 
diff --git a/Engine/Source/Runtime/NetworkReplayStreaming/NullNetworkReplayStreaming/Public/NullNetworkReplayStreaming.h b/Engine/Source/Runtime/NetworkReplayStreaming/NullNetworkReplayStreaming/Public/NullNetworkReplayStreaming.h
--- a/Engine/Source/Runtime/NetworkReplayStreaming/NullNetworkReplayStreaming/Public/NullNetworkReplayStreaming.h
+++ b/Engine/Source/Runtime/NetworkReplayStreaming/NullNetworkReplayStreaming/Public/NullNetworkReplayStreaming.h
@@ -7,17 +7,20 @@
 #include "ModuleManager.h"
 #include "UniquePtr.h"
 #include "OnlineJsonSerializer.h"
+#include "Tickable.h"
 
 /* Class to hold metadata about an entire replay */
 class FNullReplayInfo : public FOnlineJsonSerializable
 {
 public:
-	FNullReplayInfo() : LengthInMS(0), NetworkVersion(0), Changelist(0) {}
+	FNullReplayInfo() : LengthInMS(0), NetworkVersion(0), Changelist(0), DemoFileLastOffset(0), bIsValid(false) {}
 
 	int32		LengthInMS;
 	uint32		NetworkVersion;
 	uint32		Changelist;
 	FString		FriendlyName;
+	int32		DemoFileLastOffset;
+	bool		bIsValid;
 
 	// FOnlineJsonSerializable
 	BEGIN_ONLINE_JSON_SERIALIZER
@@ -25,16 +28,18 @@ public:
 		ONLINE_JSON_SERIALIZE( "NetworkVersion",	NetworkVersion );
 		ONLINE_JSON_SERIALIZE( "Changelist",		Changelist );
 		ONLINE_JSON_SERIALIZE( "FriendlyName",		FriendlyName );
+		ONLINE_JSON_SERIALIZE( "DemoFileLastOffset",	DemoFileLastOffset );
 	END_ONLINE_JSON_SERIALIZER
 };
 
 /** Default streamer that goes straight to the HD */
-class FNullNetworkReplayStreamer : public INetworkReplayStreamer
+class FNullNetworkReplayStreamer : public INetworkReplayStreamer, public FTickableGameObject
 {
 public:
 	FNullNetworkReplayStreamer() :
 		StreamerState( EStreamerState::Idle ),
-		CurrentCheckpointIndex( 0 )
+		CurrentCheckpointIndex( 0 ),
+		LastKnownFileSize( 0 )
 	{}
 	
 	/** INetworkReplayStreamer implementation */
@@ -49,7 +54,7 @@ public:
 	virtual FArchive* GetMetadataArchive() override;
 	virtual void UpdateTotalDemoTime( uint32 TimeInMS ) override;
 	virtual uint32 GetTotalDemoTime() const override { return ReplayInfo.LengthInMS; }
-	virtual bool IsDataAvailable() const override { return true; }
+	virtual bool IsDataAvailable() const override;
 	virtual void SetHighPriorityTimeRange( const uint32 StartTimeInMS, const uint32 EndTimeInMS ) override { }
 	virtual bool IsDataAvailableForTimeRange( const uint32 StartTimeInMS, const uint32 EndTimeInMS ) override { return true; }
 	virtual bool IsLoadingCheckpoint() const override { return false; }
@@ -64,12 +69,26 @@ public:
 	virtual void RequestEventData(const FString& EventID, const FOnRequestEventDataComplete& RequestEventDataComplete) override;
 	virtual void SearchEvents(const FString& EventGroup, const FOnEnumerateStreamsComplete& Delegate) override;
 
+	/** FTickableObjectBase implementation */
+	virtual void Tick(float DeltaSeconds) override;
+	virtual bool IsTickable() const override { return true; }
+	virtual TStatId GetStatId() const override;
+
+	/** FTickableGameObject implementation */
+	virtual bool IsTickableWhenPaused() const override { return true; }
+
 private:
 	bool IsNamedStreamLive( const FString& StreamName ) const;
 
 	/** Handles the details of loading a checkpoint */
 	void GotoCheckpointIndexInternal(int32 CheckpointIndex, const FOnCheckpointReadyDelegate& Delegate, int32 TimeInMS);
 
+	/** Reopen the file to refresh it's size, since file-based FArchives do not appear to update their size if they're being written to. */
+	void ReopenStreamFileForReading();
+
+	/** Overwrites the cached ReplayInfo only if the read succeeded */
+	void UpdateReplayInfoIfValid();
+
 	/** Handle to the archive that will read/write the demo header */
 	TUniquePtr<FArchive> HeaderAr;
 
@@ -101,6 +120,9 @@ private:
 
 	/** Currently playing or recording replay metadata */
 	FNullReplayInfo ReplayInfo;
+
+	/** Last known size of the replay stream file. */
+	int64 LastKnownFileSize;
 };
 
 class FNullNetworkReplayStreamingFactory : public INetworkReplayStreamingFactory
