Allow overriding the replay streamer module by name at runtime.

--------
Integrated using branch UE4-Orion-To-UE4 of change#2677662 by [REDACTED_PERSON] on 2015/09/02 16:54:01.

diff --git a/Engine/Source/Runtime/Engine/Private/DemoNetDriver.cpp b/Engine/Source/Runtime/Engine/Private/DemoNetDriver.cpp
--- a/Engine/Source/Runtime/Engine/Private/DemoNetDriver.cpp
+++ b/Engine/Source/Runtime/Engine/Private/DemoNetDriver.cpp
@@ -306,7 +306,8 @@ bool UDemoNetDriver::InitBase( bool bInitAsClient, FNetworkNotify* InNotify, con
 
 		ResetDemoState();
 
-		ReplayStreamer = FNetworkReplayStreaming::Get().GetFactory().CreateReplayStreamer();
+		const TCHAR* const StreamerOverride = URL.GetOption(TEXT("ReplayStreamerOverride="), nullptr);
+		ReplayStreamer = FNetworkReplayStreaming::Get().GetFactory(StreamerOverride).CreateReplayStreamer();
 
 		return true;
 	}
diff --git a/Engine/Source/Runtime/NetworkReplayStreaming/NetworkReplayStreaming/Private/NetworkReplayStreaming.cpp b/Engine/Source/Runtime/NetworkReplayStreaming/NetworkReplayStreaming/Private/NetworkReplayStreaming.cpp
--- a/Engine/Source/Runtime/NetworkReplayStreaming/NetworkReplayStreaming/Private/NetworkReplayStreaming.cpp
+++ b/Engine/Source/Runtime/NetworkReplayStreaming/NetworkReplayStreaming/Private/NetworkReplayStreaming.cpp
@@ -4,10 +4,18 @@
 
 IMPLEMENT_MODULE( FNetworkReplayStreaming, NetworkReplayStreaming );
 
-INetworkReplayStreamingFactory& FNetworkReplayStreaming::GetFactory()
+INetworkReplayStreamingFactory& FNetworkReplayStreaming::GetFactory(const TCHAR* FactoryNameOverride)
 {
 	FString FactoryName = TEXT( "NullNetworkReplayStreaming" );
-	GConfig->GetString( TEXT( "NetworkReplayStreaming" ), TEXT( "DefaultFactoryName" ), FactoryName, GEngineIni );
+
+	if (FactoryNameOverride == nullptr)
+	{
+		GConfig->GetString( TEXT( "NetworkReplayStreaming" ), TEXT( "DefaultFactoryName" ), FactoryName, GEngineIni );
+	}
+	else
+	{
+		FactoryName = FactoryNameOverride;
+	}
 
 	// See if we need to forcefully fallback to the null streamer
 	if ( !FModuleManager::Get().IsModuleLoaded( *FactoryName ) )
diff --git a/Engine/Source/Runtime/NetworkReplayStreaming/NetworkReplayStreaming/Public/NetworkReplayStreaming.h b/Engine/Source/Runtime/NetworkReplayStreaming/NetworkReplayStreaming/Public/NetworkReplayStreaming.h
--- a/Engine/Source/Runtime/NetworkReplayStreaming/NetworkReplayStreaming/Public/NetworkReplayStreaming.h
+++ b/Engine/Source/Runtime/NetworkReplayStreaming/NetworkReplayStreaming/Public/NetworkReplayStreaming.h
@@ -231,5 +231,5 @@ public:
 		return FModuleManager::LoadModuleChecked< FNetworkReplayStreaming >( "NetworkReplayStreaming" );
 	}
 
-	virtual INetworkReplayStreamingFactory& GetFactory();
+	virtual INetworkReplayStreamingFactory& GetFactory(const TCHAR* FactoryNameOverride = nullptr);
 };
