Fixes for Orion Debug Editor config, probably other debug configs

diff --git a/Engine/Source/Runtime/Engine/Classes/Sound/SoundWave.h b/Engine/Source/Runtime/Engine/Classes/Sound/SoundWave.h
--- a/Engine/Source/Runtime/Engine/Classes/Sound/SoundWave.h
+++ b/Engine/Source/Runtime/Engine/Classes/Sound/SoundWave.h
@@ -6,8 +6,10 @@
  * Playable sound object for raw wave files
  */
 #include "BulkData.h"
+#include "Engine/EngineTypes.h"
 #include "Sound/SoundBase.h"
 #include "Sound/SoundGroups.h"
+
 #include "SoundWave.generated.h"
 
 struct FActiveSound;
diff --git a/Engine/Source/Runtime/NetworkReplayStreaming/NullNetworkReplayStreaming/NullNetworkReplayStreaming.build.cs b/Engine/Source/Runtime/NetworkReplayStreaming/NullNetworkReplayStreaming/NullNetworkReplayStreaming.build.cs
--- a/Engine/Source/Runtime/NetworkReplayStreaming/NullNetworkReplayStreaming/NullNetworkReplayStreaming.build.cs
+++ b/Engine/Source/Runtime/NetworkReplayStreaming/NullNetworkReplayStreaming/NullNetworkReplayStreaming.build.cs
@@ -10,15 +10,13 @@ namespace UnrealBuildTool.Rules
 
             PrivateIncludePathModuleNames.Add("OnlineSubsystem");
 
-			PublicIncludePathModuleNames.Add("Engine");
-
 			PrivateDependencyModuleNames.AddRange(
 				new string[]
 				{
 					"Core",
+					"Engine",
 					"NetworkReplayStreaming",
                     "Json",
-					"Engine"
 				}
 			);
 		}
diff --git a/Engine/Source/Runtime/NetworkReplayStreaming/NullNetworkReplayStreaming/Public/NullNetworkReplayStreaming.h b/Engine/Source/Runtime/NetworkReplayStreaming/NullNetworkReplayStreaming/Public/NullNetworkReplayStreaming.h
--- a/Engine/Source/Runtime/NetworkReplayStreaming/NullNetworkReplayStreaming/Public/NullNetworkReplayStreaming.h
+++ b/Engine/Source/Runtime/NetworkReplayStreaming/NullNetworkReplayStreaming/Public/NullNetworkReplayStreaming.h
@@ -4,6 +4,7 @@
 
 #include "NetworkReplayStreaming.h"
 #include "Core.h"
+#include "Engine.h"
 #include "ModuleManager.h"
 #include "UniquePtr.h"
 #include "OnlineJsonSerializer.h"
