#orion, merge CL 2694183: Don't stop uploading replays when custom event upload fails

diff --git a/Engine/Source/Runtime/NetworkReplayStreaming/HttpNetworkReplayStreaming/Private/HttpNetworkReplayStreaming.cpp b/Engine/Source/Runtime/NetworkReplayStreaming/HttpNetworkReplayStreaming/Private/HttpNetworkReplayStreaming.cpp
--- a/Engine/Source/Runtime/NetworkReplayStreaming/HttpNetworkReplayStreaming/Private/HttpNetworkReplayStreaming.cpp
+++ b/Engine/Source/Runtime/NetworkReplayStreaming/HttpNetworkReplayStreaming/Private/HttpNetworkReplayStreaming.cpp
@@ -1179,12 +1179,8 @@ void FHttpNetworkReplayStreamer::HttpUploadCustomEventFinished(FHttpRequestPtr H
 	else
 	{
 		UE_LOG(LogHttpReplay, Error, TEXT("FHttpNetworkReplayStreamer::HttpUploadCustomEventFinished. FAILED, Response code: %d"), HttpResponse.IsValid() ? HttpResponse->GetResponseCode() : 0);
-
-		// Don't care about the error coming down
-		if (StreamerState == EStreamerState::StreamingUp)
-		{
-			SetLastError(ENetworkReplayError::ServiceUnavailable);
-		}
+		// Don't disconect service here, just report the failure
+		//SetLastError(ENetworkReplayError::ServiceUnavailable);
 	}
 }
 
