#merge, replays: Fix warning when uploading custom events with v2 api changes

diff --git a/Engine/Source/Runtime/NetworkReplayStreaming/HttpNetworkReplayStreaming/Private/HttpNetworkReplayStreaming.cpp b/Engine/Source/Runtime/NetworkReplayStreaming/HttpNetworkReplayStreaming/Private/HttpNetworkReplayStreaming.cpp
--- a/Engine/Source/Runtime/NetworkReplayStreaming/HttpNetworkReplayStreaming/Private/HttpNetworkReplayStreaming.cpp
+++ b/Engine/Source/Runtime/NetworkReplayStreaming/HttpNetworkReplayStreaming/Private/HttpNetworkReplayStreaming.cpp
@@ -1157,7 +1157,7 @@ void FHttpNetworkReplayStreamer::HttpUploadCheckpointFinished( FHttpRequestPtr H
 {
 	RequestFinished( EStreamerState::StreamingUp, EQueuedHttpRequestType::UploadingCheckpoint, HttpRequest );
 
-	if ( bSucceeded && HttpResponse->GetResponseCode() == EHttpResponseCodes::NoContent )
+	if ( bSucceeded && ( HttpResponse->GetResponseCode() == EHttpResponseCodes::Ok || HttpResponse->GetResponseCode() == EHttpResponseCodes::NoContent ) )
 	{
 		UE_LOG( LogHttpReplay, Verbose, TEXT( "FHttpNetworkReplayStreamer::HttpUploadCheckpointFinished." ) );
 	}
@@ -1172,7 +1172,7 @@ void FHttpNetworkReplayStreamer::HttpUploadCustomEventFinished(FHttpRequestPtr H
 {
 	RequestFinished(StreamerState, EQueuedHttpRequestType::UploadingCustomEvent, HttpRequest);
 
-	if (bSucceeded && HttpResponse->GetResponseCode() == EHttpResponseCodes::Ok)
+	if ( bSucceeded && ( HttpResponse->GetResponseCode() == EHttpResponseCodes::Ok || HttpResponse->GetResponseCode() == EHttpResponseCodes::NoContent ) )
 	{
 		UE_LOG(LogHttpReplay, Verbose, TEXT("FHttpNetworkReplayStreamer::HttpUploadCustomEventFinished."));
 	}
