#merge CL 2703103, replays: Add api to allow replays to be marked to not be purged + add ability to add custom events, which are used to delay the need to know the SessionName immediately.

diff --git a/Engine/Source/Runtime/NetworkReplayStreaming/HttpNetworkReplayStreaming/Private/HttpNetworkReplayStreaming.cpp b/Engine/Source/Runtime/NetworkReplayStreaming/HttpNetworkReplayStreaming/Private/HttpNetworkReplayStreaming.cpp
--- a/Engine/Source/Runtime/NetworkReplayStreaming/HttpNetworkReplayStreaming/Private/HttpNetworkReplayStreaming.cpp
+++ b/Engine/Source/Runtime/NetworkReplayStreaming/HttpNetworkReplayStreaming/Private/HttpNetworkReplayStreaming.cpp
@@ -307,6 +307,13 @@ void FHttpNetworkReplayStreamer::AddRequestToQueue( const EQueuedHttpRequestType
 	QueuedHttpRequests.Enqueue( TSharedPtr< FQueuedHttpRequest >( new FQueuedHttpRequest( Type, Request ) ) );
 }
 
+void FHttpNetworkReplayStreamer::AddCustomRequestToQueue( TSharedPtr< FQueuedHttpRequest > Request )
+{
+	UE_LOG( LogHttpReplay, Verbose, TEXT( "FHttpNetworkReplayStreamer::AddCustomRequestToQueue. Type: %s" ), EQueuedHttpRequestType::ToString( Request->Type ) );
+
+	QueuedHttpRequests.Enqueue( Request );
+}
+
 void FHttpNetworkReplayStreamer::StopStreaming()
 {
 	if ( StartStreamingDelegate.IsBound() )
@@ -557,7 +564,7 @@ void FHttpNetworkReplayStreamer::RequestEventData(const FString& EventID, const
 
 void FHttpNetworkReplayStreamer::HttpRequestEventDataFinished(FHttpRequestPtr HttpRequest, FHttpResponsePtr HttpResponse, bool bSucceeded, FOnRequestEventDataComplete RequestEventDataCompleteDelegate)
 {
-	RequestFinished(EStreamerState::StreamingDown, EQueuedHttpRequestType::RequestEventData, HttpRequest);
+	RequestFinished(StreamerState, EQueuedHttpRequestType::RequestEventData, HttpRequest);
 
 	check(RequestEventDataCompleteDelegate.IsBound());
 
@@ -655,10 +662,36 @@ void FHttpNetworkReplayStreamer::FlushCheckpointInternal( uint32 TimeInMS )
 	AddRequestToQueue( EQueuedHttpRequestType::UploadingCheckpoint, HttpRequest );
 }
 
+FQueuedHttpRequestAddEvent::FQueuedHttpRequestAddEvent( const uint32 InTimeInMS, const FString& InGroup, const FString& InMeta, const TArray<uint8>& InData, TSharedRef< class IHttpRequest > InHttpRequest ) : FQueuedHttpRequest( EQueuedHttpRequestType::UploadingCustomEvent, InHttpRequest )
+{
+	Request->SetVerb( TEXT( "POST" ) );
+	Request->SetHeader( TEXT( "Content-Type" ), TEXT( "application/octet-stream" ) );
+	Request->SetContent( InData );
+
+	TimeInMS	= InTimeInMS;
+	Group		= InGroup;
+	Meta		= InMeta;
+}
+
+bool FQueuedHttpRequestAddEvent::PreProcess( const FString& ServerURL, const FString& SessionName )
+{
+	if ( SessionName.IsEmpty() )
+	{
+		UE_LOG( LogHttpReplay, Warning, TEXT( "FQueuedHttpRequestAddEvent::PreProcess. SessionName is empty." ) );
+		return false;
+	}
+
+	// Now that we have the session name, we can set the URL
+	Request->SetURL( FString::Printf( TEXT( "%sreplay/%s/event?group=%s&time1=%i&time2=%i&meta=%s" ), *ServerURL, *SessionName, *Group, TimeInMS, TimeInMS, *Meta ) );
+
+	return true;
+}
+
 void FHttpNetworkReplayStreamer::AddEvent( const uint32 TimeInMS, const FString& Group, const FString& Meta, const TArray<uint8>& Data )
 {
-	if (SessionName.IsEmpty() || (StreamerState != EStreamerState::StreamingUp && StreamerState != EStreamerState::StreamingDown))
+	if (StreamerState != EStreamerState::StreamingUp && StreamerState != EStreamerState::StreamingDown)
 	{
+		UE_LOG( LogHttpReplay, Warning, TEXT( "FHttpNetworkReplayStreamer::AddEvent. Not streaming." ) );
 		return;
 	}
 
@@ -666,16 +699,12 @@ void FHttpNetworkReplayStreamer::AddEvent( const uint32 TimeInMS, const FString&
 	UE_LOG(LogHttpReplay, Verbose, TEXT("FHttpNetworkReplayStreamer::AddEvent. Size: %i, StreamChunkIndex: %i"), Data.Num(), StreamChunkIndex);
 
 	// Create the Http request and add to pending request list
-	TSharedRef<class IHttpRequest> HttpRequest = FHttpModule::Get().CreateRequest();
+	TSharedRef< class IHttpRequest > HttpRequest = FHttpModule::Get().CreateRequest();
 
-	HttpRequest->OnProcessRequestComplete().BindRaw(this, &FHttpNetworkReplayStreamer::HttpUploadCustomEventFinished);
+	HttpRequest->OnProcessRequestComplete().BindRaw( this, &FHttpNetworkReplayStreamer::HttpUploadCustomEventFinished );
 
-	HttpRequest->SetURL(FString::Printf(TEXT("%sreplay/%s/event?group=%s&time1=%i&time2=%i&meta=%s"), *ServerURL, *SessionName, *Group, TimeInMS, TimeInMS, *Meta));
-	HttpRequest->SetVerb(TEXT("POST"));
-	HttpRequest->SetHeader(TEXT("Content-Type"), TEXT("application/octet-stream"));
-	HttpRequest->SetContent(Data);
-	
-	AddRequestToQueue(EQueuedHttpRequestType::UploadingCustomEvent, HttpRequest);
+	// Add it as a custom event so we can snag the session name at the time of send (which we should have by then)
+	AddCustomRequestToQueue( TSharedPtr< FQueuedHttpRequest >( new FQueuedHttpRequestAddEvent( TimeInMS, Group, Meta, Data, HttpRequest ) ) );
 }
 
 void FHttpNetworkReplayStreamer::DownloadHeader()
@@ -749,6 +778,38 @@ void FHttpNetworkReplayStreamer::ConditionallyDownloadNextChunk()
 	LastChunkTime = FPlatformTime::Seconds();
 }
 
+void FHttpNetworkReplayStreamer::KeepReplay( const FString& ReplayName, const bool bKeep )
+{
+	TSharedRef<class IHttpRequest> HttpRequest = FHttpModule::Get().CreateRequest();
+
+	// Download the next stream chunk
+	if ( bKeep )
+	{
+		HttpRequest->SetURL( FString::Printf( TEXT( "%sreplay/%s?shouldKeep=true" ), *ServerURL, *ReplayName ) );
+	}
+	else
+	{
+		HttpRequest->SetURL( FString::Printf( TEXT( "%sreplay/%s?shouldKeep=false" ), *ServerURL, *ReplayName ) );
+	}
+
+	HttpRequest->SetVerb( TEXT( "POST" ) );
+	HttpRequest->SetHeader( TEXT( "Content-Type" ), TEXT( "application/json" ) );
+
+	HttpRequest->OnProcessRequestComplete().BindRaw( this, &FHttpNetworkReplayStreamer::KeepReplayFinished );
+
+	AddRequestToQueue( EQueuedHttpRequestType::KeepReplay, HttpRequest );
+}
+
+void FHttpNetworkReplayStreamer::KeepReplayFinished( FHttpRequestPtr HttpRequest, FHttpResponsePtr HttpResponse, bool bSucceeded )
+{
+	RequestFinished( StreamerState, EQueuedHttpRequestType::KeepReplay, HttpRequest );
+
+	if ( !bSucceeded || HttpResponse->GetResponseCode() != EHttpResponseCodes::NoContent )
+	{
+		UE_LOG( LogHttpReplay, Error, TEXT( "FHttpNetworkReplayStreamer::KeepReplayFinished. FAILED, Response code: %d" ), HttpResponse.IsValid() ? HttpResponse->GetResponseCode() : 0 );
+	}
+}
+
 void FHttpNetworkReplayStreamer::RefreshViewer( const bool bFinal )
 {
 	TSharedRef<class IHttpRequest> HttpRequest = FHttpModule::Get().CreateRequest();
@@ -1043,16 +1104,20 @@ void FHttpNetworkReplayStreamer::ConditionallyEnumerateCheckpoints()
 
 void FHttpNetworkReplayStreamer::EnumerateEvents(const FString& Group, const FEnumerateEventsCompleteDelegate& EnumerationCompleteDelegate)
 {
+	EnumerateEvents( SessionName, Group, EnumerationCompleteDelegate );
+}
 
+void FHttpNetworkReplayStreamer::EnumerateEvents( const FString& ReplayName, const FString& Group, const FEnumerateEventsCompleteDelegate& EnumerationCompleteDelegate )
+{
 	TSharedRef<class IHttpRequest> HttpRequest = FHttpModule::Get().CreateRequest();
 
-	// Enumerate all of the sessions
-	HttpRequest->SetURL(FString::Printf(TEXT("%sreplay/%s/event?group=%s"), *ServerURL, *SessionName, *Group));
-	HttpRequest->SetVerb(TEXT("GET"));
+	// Enumerate all of the events
+	HttpRequest->SetURL( FString::Printf( TEXT( "%sreplay/%s/event?group=%s" ), *ServerURL, *ReplayName, *Group ) );
+	HttpRequest->SetVerb( TEXT( "GET" ) );
 
-	HttpRequest->OnProcessRequestComplete().BindRaw(this, &FHttpNetworkReplayStreamer::HttpEnumerateEventsFinished, EnumerationCompleteDelegate);
+	HttpRequest->OnProcessRequestComplete().BindRaw( this, &FHttpNetworkReplayStreamer::HttpEnumerateEventsFinished, EnumerationCompleteDelegate );
 
-	AddRequestToQueue(EQueuedHttpRequestType::EnumeratingCustomEvent, HttpRequest);
+	AddRequestToQueue( EQueuedHttpRequestType::EnumeratingCustomEvent, HttpRequest );
 }
 
 void FHttpNetworkReplayStreamer::RequestFinished( EStreamerState ExpectedStreamerState, EQueuedHttpRequestType::Type ExpectedType, FHttpRequestPtr HttpRequest )
@@ -1519,7 +1584,7 @@ void FHttpNetworkReplayStreamer::HttpEnumerateCheckpointsFinished( FHttpRequestP
 
 void FHttpNetworkReplayStreamer::HttpEnumerateEventsFinished(FHttpRequestPtr HttpRequest, FHttpResponsePtr HttpResponse, bool bSucceeded, FEnumerateEventsCompleteDelegate EnumerateEventsDelegate)
 {
-	RequestFinished(EStreamerState::StreamingDown, EQueuedHttpRequestType::EnumeratingCustomEvent, HttpRequest);
+	RequestFinished(StreamerState, EQueuedHttpRequestType::EnumeratingCustomEvent, HttpRequest);
 
 	if (bSucceeded && HttpResponse->GetResponseCode() == EHttpResponseCodes::Ok)
 	{
@@ -1598,6 +1663,12 @@ bool FHttpNetworkReplayStreamer::ProcessNextHttpRequest()
 			return ProcessNextHttpRequest();
 		}
 
+		if ( !QueuedRequest->PreProcess( ServerURL, SessionName ) )
+		{
+			// This request failed, go ahead and process the next request
+			return ProcessNextHttpRequest();
+		}
+
 		InFlightHttpRequest = QueuedRequest;
 
 		ProcessRequestInternal( InFlightHttpRequest->Request );
diff --git a/Engine/Source/Runtime/NetworkReplayStreaming/HttpNetworkReplayStreaming/Public/HttpNetworkReplayStreaming.h b/Engine/Source/Runtime/NetworkReplayStreaming/HttpNetworkReplayStreaming/Public/HttpNetworkReplayStreaming.h
--- a/Engine/Source/Runtime/NetworkReplayStreaming/HttpNetworkReplayStreaming/Public/HttpNetworkReplayStreaming.h
+++ b/Engine/Source/Runtime/NetworkReplayStreaming/HttpNetworkReplayStreaming/Public/HttpNetworkReplayStreaming.h
@@ -47,6 +47,7 @@ namespace EQueuedHttpRequestType
 		RequestEventData,			// We are in the process of requesting the data for a specific event
 		UploadHeader,				// Request to upload header (has to be done after we get info from server)
 		StopStreaming,				// Request to stop streaming
+		KeepReplay,					// Request to keep replay (or cancel keeping replay)
 	};
 
 	inline const TCHAR* ToString( EQueuedHttpRequestType::Type Type )
@@ -89,6 +90,8 @@ namespace EQueuedHttpRequestType
 				return TEXT( "UploadHeader" );
 			case StopStreaming:
 				return TEXT( "StopStreaming" );
+			case KeepReplay:
+				return TEXT( "KeepReplay" );
 		}
 
 		return TEXT( "Unknown EQueuedHttpRequestType type." );
@@ -102,8 +105,37 @@ public:
 	{
 	}
 
+	virtual ~FQueuedHttpRequest()
+	{
+	}
+
 	EQueuedHttpRequestType::Type		Type;
 	TSharedPtr< class IHttpRequest >	Request;
+
+	virtual bool PreProcess( const FString& ServerURL, const FString& SessionName )
+	{
+		return true;
+	}
+};
+
+/**
+* FQueuedHttpRequestAddEvent
+* Custom event so that we can defer the need to knowing SessionName until we actually send it (which we should have it by then, since requests are executed in order)
+*/
+class FQueuedHttpRequestAddEvent : public FQueuedHttpRequest
+{
+public:
+	FQueuedHttpRequestAddEvent( const uint32 InTimeInMS, const FString& InGroup, const FString& InMeta, const TArray<uint8>& InData, TSharedRef< class IHttpRequest > InHttpRequest );
+
+	virtual ~FQueuedHttpRequestAddEvent()
+	{
+	}
+
+	virtual bool PreProcess( const FString& ServerURL, const FString& SessionName ) override;
+
+	uint32		TimeInMS;
+	FString		Group;
+	FString		Meta;
 };
 
 /**
@@ -133,10 +165,13 @@ public:
 	virtual bool		IsLive() const override;
 	virtual void		DeleteFinishedStream( const FString& StreamName, const FOnDeleteFinishedStreamComplete& Delegate ) const override;
 	virtual void		EnumerateStreams( const FNetworkReplayVersion& ReplayVersion, const FString& UserString, const FString& MetaString, const FOnEnumerateStreamsComplete& Delegate ) override;
+	virtual void		EnumerateEvents( const FString& Group, const FEnumerateEventsCompleteDelegate& EnumerationCompleteDelegate ) override;
+	virtual void		EnumerateEvents( const FString& ReplayName, const FString& Group, const FEnumerateEventsCompleteDelegate& EnumerationCompleteDelegate ) override;
 	virtual void		EnumerateRecentStreams( const FNetworkReplayVersion& ReplayVersion, const FString& RecentViewer, const FOnEnumerateStreamsComplete& Delegate ) override;
 	virtual void		AddUserToReplay(const FString& UserString);
 	virtual void		RequestEventData(const FString& EventId, const FOnRequestEventDataComplete& Delegate) override;
 	virtual void		SearchEvents(const FString& EventGroup, const FOnEnumerateStreamsComplete& Delegate) override;
+	virtual void		KeepReplay( const FString& ReplayName, const bool bKeep ) override;
 	virtual ENetworkReplayError::Type GetLastError() const override;
 
 	/** FHttpNetworkReplayStreamer */
@@ -153,9 +188,9 @@ public:
 	void FlushCheckpointInternal( uint32 TimeInMS );
 	void AddEvent( const uint32 TimeInMS, const FString& Group, const FString& Meta, const TArray<uint8>& Data );
 	void AddRequestToQueue( const EQueuedHttpRequestType::Type Type, TSharedPtr< class IHttpRequest >	Request );
+	void AddCustomRequestToQueue( TSharedPtr< FQueuedHttpRequest > Request );
 	void EnumerateCheckpoints();
 	void ConditionallyEnumerateCheckpoints();
-	void EnumerateEvents( const FString& Group, const FEnumerateEventsCompleteDelegate& EnumerationCompleteDelegate );
 
 	virtual void ProcessRequestInternal( TSharedPtr< class IHttpRequest > Request );
 
@@ -186,6 +221,7 @@ public:
 	void HttpEnumerateEventsFinished( FHttpRequestPtr HttpRequest, FHttpResponsePtr HttpResponse, bool bSucceeded, FEnumerateEventsCompleteDelegate EnumerateEventsDelegate );
 	void HttpAddUserFinished( FHttpRequestPtr HttpRequest, FHttpResponsePtr HttpResponse, bool bSucceeded );
 	void HttpRequestEventDataFinished( FHttpRequestPtr HttpRequest, FHttpResponsePtr HttpResponse, bool bSucceeded, FOnRequestEventDataComplete RequestEventDataCompleteDelegate );
+	void KeepReplayFinished( FHttpRequestPtr HttpRequest, FHttpResponsePtr HttpResponse, bool bSucceeded );
 
 	bool ProcessNextHttpRequest();
 	void Tick( const float DeltaTime );
diff --git a/Engine/Source/Runtime/NetworkReplayStreaming/NetworkReplayStreaming/Public/NetworkReplayStreaming.h b/Engine/Source/Runtime/NetworkReplayStreaming/NetworkReplayStreaming/Public/NetworkReplayStreaming.h
--- a/Engine/Source/Runtime/NetworkReplayStreaming/NetworkReplayStreaming/Public/NetworkReplayStreaming.h
+++ b/Engine/Source/Runtime/NetworkReplayStreaming/NetworkReplayStreaming/Public/NetworkReplayStreaming.h
@@ -177,8 +177,10 @@ public:
 	virtual bool IsLoadingCheckpoint() const = 0;
 	virtual void AddEvent( const uint32 TimeInMS, const FString& Group, const FString& Meta, const TArray<uint8>& Data ) = 0;
 	virtual void EnumerateEvents( const FString& Group, const FEnumerateEventsCompleteDelegate& EnumerationCompleteDelegate ) = 0;
-	virtual void RequestEventData(const FString& EventID, const FOnRequestEventDataComplete& RequestEventDataComplete) = 0;
+	virtual void EnumerateEvents( const FString& ReplayName, const FString& Group, const FEnumerateEventsCompleteDelegate& EnumerationCompleteDelegate ) = 0;
+	virtual void RequestEventData( const FString& EventID, const FOnRequestEventDataComplete& RequestEventDataComplete ) = 0;
 	virtual void SearchEvents(const FString& EventGroup, const FOnEnumerateStreamsComplete& Delegate) = 0;
+	virtual void KeepReplay( const FString& ReplayName, const bool bKeep ) = 0;
 	/** Returns true if the playing stream is currently in progress */
 	virtual bool IsLive() const = 0;
 
diff --git a/Engine/Source/Runtime/NetworkReplayStreaming/NullNetworkReplayStreaming/Public/NullNetworkReplayStreaming.h b/Engine/Source/Runtime/NetworkReplayStreaming/NullNetworkReplayStreaming/Public/NullNetworkReplayStreaming.h
--- a/Engine/Source/Runtime/NetworkReplayStreaming/NullNetworkReplayStreaming/Public/NullNetworkReplayStreaming.h
+++ b/Engine/Source/Runtime/NetworkReplayStreaming/NullNetworkReplayStreaming/Public/NullNetworkReplayStreaming.h
@@ -67,8 +67,10 @@ public:
 	virtual void AddUserToReplay(const FString& UserString) override;
 	virtual void AddEvent(const uint32 TimeInMS, const FString& Group, const FString& Meta, const TArray<uint8>& Data) override;
 	virtual void EnumerateEvents(const FString& Group, const FEnumerateEventsCompleteDelegate& EnumerationCompleteDelegate) override;
+	virtual void EnumerateEvents( const FString& ReplayName, const FString& Group, const FEnumerateEventsCompleteDelegate& EnumerationCompleteDelegate ) override {}
 	virtual void RequestEventData(const FString& EventID, const FOnRequestEventDataComplete& RequestEventDataComplete) override;
 	virtual void SearchEvents(const FString& EventGroup, const FOnEnumerateStreamsComplete& Delegate) override;
+	virtual void KeepReplay( const FString& ReplayName, const bool bKeep ) override {}
 
 	/** FTickableObjectBase implementation */
 	virtual void Tick(float DeltaSeconds) override;
