In-memory network replay streamer implementation - useful for storing local, transient replays without hitting the disk or network.

--------
Integrated using branch UE4-Orion-To-UE4 of change#2704828 by [REDACTED_PERSON] on 2015/09/24 14:19:13.

diff --git a/Engine/Source/Runtime/NetworkReplayStreaming/InMemoryNetworkReplayStreaming/InMemoryNetworkReplayStreaming.build.cs b/Engine/Source/Runtime/NetworkReplayStreaming/InMemoryNetworkReplayStreaming/InMemoryNetworkReplayStreaming.build.cs
new file mode 100644
--- /dev/null
+++ b/Engine/Source/Runtime/NetworkReplayStreaming/InMemoryNetworkReplayStreaming/InMemoryNetworkReplayStreaming.build.cs
@@ -0,0 +1,22 @@
+// Copyright 1998-2015 Epic Games, Inc. All Rights Reserved.
+
+namespace UnrealBuildTool.Rules
+{
+	public class InMemoryNetworkReplayStreaming : ModuleRules
+	{
+		public InMemoryNetworkReplayStreaming( TargetInfo Target )
+		{
+			PublicIncludePathModuleNames.Add("Engine");
+
+			PrivateDependencyModuleNames.AddRange(
+				new string[]
+				{
+					"Core",
+					"NetworkReplayStreaming",
+                    "Json",
+					"Engine"
+				}
+			);
+		}
+	}
+}
\ No newline at end of file
diff --git a/Engine/Source/Runtime/NetworkReplayStreaming/InMemoryNetworkReplayStreaming/Private/InMemoryNetworkReplayStreaming.cpp b/Engine/Source/Runtime/NetworkReplayStreaming/InMemoryNetworkReplayStreaming/Private/InMemoryNetworkReplayStreaming.cpp
new file mode 100644
--- /dev/null
+++ b/Engine/Source/Runtime/NetworkReplayStreaming/InMemoryNetworkReplayStreaming/Private/InMemoryNetworkReplayStreaming.cpp
@@ -0,0 +1,413 @@
+// Copyright 1998-2015 Epic Games, Inc. All Rights Reserved.
+
+#include "InMemoryNetworkReplayStreaming.h"
+#include "Paths.h"
+#include "EngineVersion.h"
+#include "Guid.h"
+#include "DateTime.h"
+
+DEFINE_LOG_CATEGORY_STATIC( LogMemoryReplay, Log, All );
+
+static FString GetAutomaticDemoName()
+{
+	return FGuid::NewGuid().ToString();
+}
+
+void FInMemoryNetworkReplayStreamer::StartStreaming( const FString& CustomName, const FString& FriendlyName, const TArray< FString >& UserNames, bool bRecord, const FNetworkReplayVersion& ReplayVersion, const FOnStreamReadyDelegate& Delegate )
+{
+	if ( CustomName.IsEmpty() )
+	{
+		if ( bRecord )
+		{
+			// If we're recording and the caller didn't provide a name, generate one automatically
+			CurrentStreamName = GetAutomaticDemoName();
+		}
+		else
+		{
+			// Can't play a replay if the user didn't provide a name!
+			Delegate.ExecuteIfBound( false, bRecord );
+			return;
+		}
+	}
+	else
+	{
+		CurrentStreamName = CustomName;
+	}
+
+	if ( !bRecord )
+	{
+		FInMemoryReplay* FoundReplay = GetCurrentReplay();
+		if (FoundReplay == nullptr)
+		{
+			Delegate.ExecuteIfBound(false, bRecord);
+			return;
+		}
+
+		FileAr.Reset(new FMemoryReader(FoundReplay->Stream));
+		HeaderAr.Reset(new FMemoryReader(FoundReplay->Header));
+		StreamerState = EStreamerState::Playback;
+	}
+	else
+	{
+		// Add or overwrite a demo with this name
+		TUniquePtr<FInMemoryReplay> NewReplay(new FInMemoryReplay);
+
+		NewReplay->StreamInfo.Name = CurrentStreamName;
+		NewReplay->StreamInfo.FriendlyName = FriendlyName;
+		NewReplay->StreamInfo.Timestamp = FDateTime::Now();
+		NewReplay->StreamInfo.bIsLive = true;
+		NewReplay->StreamInfo.Changelist = ReplayVersion.Changelist;
+		NewReplay->NetworkVersion = ReplayVersion.NetworkVersion;
+
+		// Open archives for writing
+		FileAr.Reset(new FMemoryWriter(NewReplay->Stream));
+		HeaderAr.Reset(new FMemoryWriter(NewReplay->Header));
+
+		OwningFactory->Replays.Add(CurrentStreamName, MoveTemp(NewReplay));
+		
+		StreamerState = EStreamerState::Recording;
+	}
+
+	// Notify immediately
+	Delegate.ExecuteIfBound( FileAr.Get() != nullptr && HeaderAr.Get() != nullptr, bRecord );
+}
+
+void FInMemoryNetworkReplayStreamer::StopStreaming()
+{
+	if (StreamerState == EStreamerState::Recording)
+	{
+		FInMemoryReplay* FoundReplay = GetCurrentReplayChecked();
+
+		FoundReplay->StreamInfo.SizeInBytes = FoundReplay->Header.Num() + FoundReplay->Stream.Num() + FoundReplay->Metadata.Num();
+		for(const auto& Checkpoint : FoundReplay->Checkpoints)
+		{
+			FoundReplay->StreamInfo.SizeInBytes += Checkpoint.Data.Num();
+		}
+
+		FoundReplay->StreamInfo.bIsLive = false;
+	}
+
+	HeaderAr.Reset();
+	FileAr.Reset();
+	MetadataFileAr.Reset();
+
+	CurrentStreamName.Empty();
+	StreamerState = EStreamerState::Idle;
+}
+
+FArchive* FInMemoryNetworkReplayStreamer::GetHeaderArchive()
+{
+	return HeaderAr.Get();
+}
+
+FArchive* FInMemoryNetworkReplayStreamer::GetStreamingArchive()
+{
+	return FileAr.Get();
+}
+
+FArchive* FInMemoryNetworkReplayStreamer::GetMetadataArchive()
+{
+	check( StreamerState != EStreamerState::Idle );
+
+	// Create the metadata archive on-demand
+	if (!MetadataFileAr)
+	{
+		FInMemoryReplay* FoundReplay = GetCurrentReplayChecked();
+
+		switch (StreamerState)
+		{
+			case EStreamerState::Recording:
+			{
+				MetadataFileAr.Reset(new FMemoryWriter(FoundReplay->Metadata));
+				break;
+			}
+			case EStreamerState::Playback:
+				MetadataFileAr.Reset(new FMemoryReader(FoundReplay->Metadata));
+				break;
+
+			default:
+				break;
+		}
+	}
+
+	return MetadataFileAr.Get();
+}
+
+void FInMemoryNetworkReplayStreamer::UpdateTotalDemoTime(uint32 TimeInMS)
+{
+	check(StreamerState == EStreamerState::Recording);
+
+	FInMemoryReplay* FoundReplay = GetCurrentReplayChecked();
+
+	FoundReplay->StreamInfo.LengthInMS = TimeInMS;
+}
+
+uint32 FInMemoryNetworkReplayStreamer::GetTotalDemoTime() const
+{
+	check(StreamerState != EStreamerState::Idle);
+
+	const FInMemoryReplay* FoundReplay = GetCurrentReplayChecked();
+
+	return FoundReplay->StreamInfo.LengthInMS;
+}
+
+bool FInMemoryNetworkReplayStreamer::IsDataAvailable() const
+{
+	// Assumptions:
+	// 1. All streamer instances run on the same thread, not simultaneously
+	// 2. A recording DemoNetDriver will write either no frames or entire frames each time it ticks
+	check(StreamerState == EStreamerState::Playback);
+
+	return FileAr.IsValid() && FileAr->Tell() < FileAr->TotalSize();
+}
+
+bool FInMemoryNetworkReplayStreamer::IsLive() const
+{
+	return IsNamedStreamLive(CurrentStreamName);
+}
+
+bool FInMemoryNetworkReplayStreamer::IsNamedStreamLive( const FString& StreamName ) const
+{
+	TUniquePtr<FInMemoryReplay>* FoundReplay = OwningFactory->Replays.Find(StreamName);
+	if (FoundReplay != nullptr)
+	{
+		return (*FoundReplay)->StreamInfo.bIsLive;
+	}
+
+	return false;
+}
+
+void FInMemoryNetworkReplayStreamer::DeleteFinishedStream( const FString& StreamName, const FOnDeleteFinishedStreamComplete& Delegate ) const
+{
+	// Danger! Deleting a stream that is still being read by another streaming instance is not supported!
+
+	// Live streams can't be deleted
+	if (IsNamedStreamLive(StreamName))
+	{
+		UE_LOG(LogMemoryReplay, Log, TEXT("Can't delete network replay stream %s because it is live!"), *StreamName);
+		Delegate.ExecuteIfBound(false);
+		return;
+	}
+
+	const int32 NumRemoved = OwningFactory->Replays.Remove(StreamName);
+
+	Delegate.ExecuteIfBound(NumRemoved > 0);
+}
+
+void FInMemoryNetworkReplayStreamer::EnumerateStreams( const FNetworkReplayVersion& ReplayVersion, const FString& UserString, const FString& MetaString, const FOnEnumerateStreamsComplete& Delegate )
+{
+	TArray<FNetworkReplayStreamInfo> Results;
+
+	for (const auto& StreamPair : OwningFactory->Replays)
+	{
+		// Check version. NetworkVersion and changelist of 0 will ignore version check.
+		const bool NetworkVersionMatches = ReplayVersion.NetworkVersion == StreamPair.Value->NetworkVersion;
+		const bool ChangelistMatches = ReplayVersion.Changelist == StreamPair.Value->StreamInfo.Changelist;
+
+		const bool NetworkVersionPasses = ReplayVersion.NetworkVersion == 0 || NetworkVersionMatches;
+		const bool ChangelistPasses = ReplayVersion.Changelist == 0 || ChangelistMatches;
+
+		if ( NetworkVersionPasses && ChangelistPasses )
+		{
+			Results.Add(StreamPair.Value->StreamInfo);
+		}
+	}
+
+	Delegate.ExecuteIfBound(Results);
+}
+
+void FInMemoryNetworkReplayStreamer::AddUserToReplay(const FString& UserString)
+{
+	UE_LOG(LogMemoryReplay, Log, TEXT("FInMemoryNetworkReplayStreamer::AddUserToReplay is currently unsupported."));
+}
+
+void FInMemoryNetworkReplayStreamer::AddEvent( const uint32 TimeInMS, const FString& Group, const FString& Meta, const TArray<uint8>& Data )
+{
+	UE_LOG(LogMemoryReplay, Log, TEXT("FInMemoryNetworkReplayStreamer::AddEvent is currently unsupported."));
+}
+
+void FInMemoryNetworkReplayStreamer::EnumerateEvents( const FString& Group, const FEnumerateEventsCompleteDelegate& EnumerationCompleteDelegate )
+{
+	UE_LOG(LogMemoryReplay, Log, TEXT("FInMemoryNetworkReplayStreamer::EnumerateEvents is currently unsupported."));
+}
+
+void FInMemoryNetworkReplayStreamer::EnumerateEvents(const FString& ReplayName, const FString& Group, const FEnumerateEventsCompleteDelegate& EnumerationCompleteDelegate)
+{
+	UE_LOG(LogMemoryReplay, Log, TEXT("FInMemoryNetworkReplayStreamer::EnumerateEvents is currently unsupported."));
+}
+
+void FInMemoryNetworkReplayStreamer::RequestEventData(const FString& EventID, const FOnRequestEventDataComplete& RequestEventDataComplete)
+{
+	UE_LOG(LogMemoryReplay, Log, TEXT("FInMemoryNetworkReplayStreamer::RequestEventData is currently unsupported."));
+}
+
+void FInMemoryNetworkReplayStreamer::SearchEvents(const FString& EventGroup, const FOnEnumerateStreamsComplete& Delegate)
+{
+	UE_LOG(LogMemoryReplay, Log, TEXT("FInMemoryNetworkReplayStreamer::SearchEvents is currently unsupported."));
+}
+
+void FInMemoryNetworkReplayStreamer::KeepReplay(const FString& ReplayName, const bool bKeep)
+{
+	UE_LOG(LogMemoryReplay, Log, TEXT("FInMemoryNetworkReplayStreamer::KeepReplay is currently unsupported."));
+}
+
+FArchive* FInMemoryNetworkReplayStreamer::GetCheckpointArchive()
+{
+	// If the archive is null, and the API is being used properly, the caller is writing a checkpoint...
+	if ( CheckpointAr.Get() == nullptr )
+	{
+		check(StreamerState != EStreamerState::Playback);
+
+		UE_LOG(LogMemoryReplay, Log, TEXT("FInMemoryNetworkReplayStreamer::GetCheckpointArchive. Creating new checkpoint."));
+
+		FInMemoryReplay* FoundReplay = GetCurrentReplayChecked();
+
+		const int32 NewCheckpointIndex = FoundReplay->Checkpoints.Add(FInMemoryReplay::FCheckpoint());
+		CheckpointAr.Reset(new FMemoryWriter(FoundReplay->Checkpoints[NewCheckpointIndex].Data));
+	}
+
+	return CheckpointAr.Get();
+}
+
+void FInMemoryNetworkReplayStreamer::FlushCheckpoint(const uint32 TimeInMS)
+{
+	UE_LOG(LogMemoryReplay, Log, TEXT("FInMemoryNetworkReplayStreamer::FlushCheckpoint. TimeInMS: %u"), TimeInMS);
+
+	check(FileAr.Get() != nullptr);
+
+	// Finalize the checkpoint data.
+	CheckpointAr.Reset();
+	
+	FInMemoryReplay* FoundReplay = GetCurrentReplayChecked();
+
+	FoundReplay->Checkpoints.Last().TimeInMS = TimeInMS;
+	FoundReplay->Checkpoints.Last().StreamByteOffset = FileAr->Tell();
+}
+
+void FInMemoryNetworkReplayStreamer::GotoCheckpointIndex(const int32 CheckpointIndex, const FOnCheckpointReadyDelegate& Delegate)
+{
+	GotoCheckpointIndexInternal(CheckpointIndex, Delegate, -1);
+}
+
+void FInMemoryNetworkReplayStreamer::GotoCheckpointIndexInternal(int32 CheckpointIndex, const FOnCheckpointReadyDelegate& Delegate, int32 TimeInMS)
+{
+	check( FileAr.Get() != nullptr);
+
+	if ( CheckpointIndex == -1 )
+	{
+		// Create a dummy checkpoint archive to indicate this is the first checkpoint
+		CheckpointAr.Reset(new FArchive);
+
+		FileAr->Seek(0);
+		
+		Delegate.ExecuteIfBound( true, TimeInMS );
+		return;
+	}
+
+	FInMemoryReplay* FoundReplay = GetCurrentReplayChecked();
+
+	if (!FoundReplay->Checkpoints.IsValidIndex(CheckpointIndex))
+	{
+		UE_LOG(LogMemoryReplay, Log, TEXT("FNullNetworkReplayStreamer::GotoCheckpointIndex. Index %i is out of bounds."), CheckpointIndex);
+		Delegate.ExecuteIfBound( false, TimeInMS );
+		return;
+	}
+
+	CheckpointAr.Reset(new FMemoryReader(FoundReplay->Checkpoints[CheckpointIndex].Data));
+
+	FileAr->Seek( FoundReplay->Checkpoints[CheckpointIndex].StreamByteOffset );
+
+	Delegate.ExecuteIfBound( true, TimeInMS );
+}
+
+FInMemoryReplay* FInMemoryNetworkReplayStreamer::GetCurrentReplay()
+{
+	TUniquePtr<FInMemoryReplay>* FoundReplay = OwningFactory->Replays.Find(CurrentStreamName);
+	if (FoundReplay != nullptr)
+	{
+		return FoundReplay->Get();
+	}
+
+	return nullptr;
+}
+
+FInMemoryReplay* FInMemoryNetworkReplayStreamer::GetCurrentReplay() const
+{
+	TUniquePtr<FInMemoryReplay>* FoundReplay = OwningFactory->Replays.Find(CurrentStreamName);
+	if (FoundReplay != nullptr)
+	{
+		return FoundReplay->Get();
+	}
+
+	return nullptr;
+}
+
+FInMemoryReplay* FInMemoryNetworkReplayStreamer::GetCurrentReplayChecked()
+{
+	FInMemoryReplay* FoundReplay = GetCurrentReplay();
+	check(FoundReplay != nullptr);
+	
+	return FoundReplay;
+}
+
+FInMemoryReplay* FInMemoryNetworkReplayStreamer::GetCurrentReplayChecked() const
+{
+	FInMemoryReplay* FoundReplay = GetCurrentReplay();
+	check(FoundReplay != nullptr);
+	
+	return FoundReplay;
+}
+
+void FInMemoryNetworkReplayStreamer::GotoTimeInMS(const uint32 TimeInMS, const FOnCheckpointReadyDelegate& Delegate)
+{
+	int32 CheckpointIndex = -1;
+
+	FInMemoryReplay* FoundReplay = GetCurrentReplayChecked();
+
+	if ( FoundReplay->Checkpoints.Num() > 0 && TimeInMS >= FoundReplay->Checkpoints.Last().TimeInMS )
+	{
+		// If we're after the very last checkpoint, that's the one we want
+		CheckpointIndex = FoundReplay->Checkpoints.Num() - 1;
+	}
+	else
+	{
+		// Checkpoints should be sorted by time, return the checkpoint that exists right before the current time
+		// For fine scrubbing, we'll fast forward the rest of the way
+		// NOTE - If we're right before the very first checkpoint, we'll return -1, which is what we want when we want to start from the very beginning
+		for ( int i = 0; i < FoundReplay->Checkpoints.Num(); i++ )
+		{
+			if ( TimeInMS < FoundReplay->Checkpoints[i].TimeInMS )
+			{
+				CheckpointIndex = i - 1;
+				break;
+			}
+		}
+	}
+
+	int32 ExtraSkipTimeInMS = TimeInMS;
+
+	if ( CheckpointIndex >= 0 )
+	{
+		// Subtract off checkpoint time so we pass in the leftover to the engine to fast forward through for the fine scrubbing part
+		ExtraSkipTimeInMS = TimeInMS - FoundReplay->Checkpoints[ CheckpointIndex ].TimeInMS;
+	}
+
+	GotoCheckpointIndexInternal( CheckpointIndex, Delegate, ExtraSkipTimeInMS );
+}
+
+void FInMemoryNetworkReplayStreamer::Tick(float DeltaSeconds)
+{
+	
+}
+
+TStatId FInMemoryNetworkReplayStreamer::GetStatId() const
+{
+	RETURN_QUICK_DECLARE_CYCLE_STAT(FInMemoryNetworkReplayStreamer, STATGROUP_Tickables);
+}
+
+IMPLEMENT_MODULE(FInMemoryNetworkReplayStreamingFactory, InMemoryNetworkReplayStreaming)
+
+TSharedPtr<INetworkReplayStreamer> FInMemoryNetworkReplayStreamingFactory::CreateReplayStreamer() 
+{
+	return TSharedPtr<INetworkReplayStreamer>(new FInMemoryNetworkReplayStreamer(this));
+}
diff --git a/Engine/Source/Runtime/NetworkReplayStreaming/InMemoryNetworkReplayStreaming/Public/InMemoryNetworkReplayStreaming.h b/Engine/Source/Runtime/NetworkReplayStreaming/InMemoryNetworkReplayStreaming/Public/InMemoryNetworkReplayStreaming.h
new file mode 100644
--- /dev/null
+++ b/Engine/Source/Runtime/NetworkReplayStreaming/InMemoryNetworkReplayStreaming/Public/InMemoryNetworkReplayStreaming.h
@@ -0,0 +1,156 @@
+// Copyright 1998-2015 Epic Games, Inc. All Rights Reserved.
+
+#pragma once
+
+#include "NetworkReplayStreaming.h"
+#include "Core.h"
+#include "ModuleManager.h"
+#include "UniquePtr.h"
+#include "OnlineJsonSerializer.h"
+#include "Tickable.h"
+
+class FInMemoryNetworkReplayStreamingFactory;
+
+/* Holds all data about an entire replay */
+struct FInMemoryReplay
+{
+	struct FCheckpoint
+	{
+		FCheckpoint() :
+			TimeInMS(0),
+			StreamByteOffset(0) {}
+
+		TArray<uint8> Data;
+		uint32 TimeInMS;
+		uint32 StreamByteOffset;
+	};
+
+	FInMemoryReplay() :
+		NetworkVersion(0) {}
+
+	TArray<uint8> Header;
+	TArray<uint8> Stream;
+	TArray<uint8> Metadata;
+	TArray<FCheckpoint> Checkpoints;
+	FNetworkReplayStreamInfo StreamInfo;
+	uint32 NetworkVersion;
+};
+
+/** Streamer that keeps all data in memory only */
+class FInMemoryNetworkReplayStreamer : public INetworkReplayStreamer, public FTickableGameObject
+{
+public:
+	FInMemoryNetworkReplayStreamer(FInMemoryNetworkReplayStreamingFactory* InFactory) :
+		OwningFactory( InFactory ),
+		StreamerState( EStreamerState::Idle )
+	{
+		check(OwningFactory != nullptr);
+	}
+	
+	/** INetworkReplayStreamer implementation */
+	virtual void StartStreaming( const FString& CustomName, const FString& FriendlyName, const TArray< FString >& UserNames, bool bRecord, const FNetworkReplayVersion& ReplayVersion, const FOnStreamReadyDelegate& Delegate ) override;
+	virtual void StopStreaming() override;
+	virtual FArchive* GetHeaderArchive() override;
+	virtual FArchive* GetStreamingArchive() override;
+	virtual FArchive* GetCheckpointArchive() override;
+	virtual void FlushCheckpoint( const uint32 TimeInMS ) override;
+	virtual void GotoCheckpointIndex( const int32 CheckpointIndex, const FOnCheckpointReadyDelegate& Delegate ) override;
+	virtual void GotoTimeInMS( const uint32 TimeInMS, const FOnCheckpointReadyDelegate& Delegate ) override;
+	virtual FArchive* GetMetadataArchive() override;
+	virtual void UpdateTotalDemoTime( uint32 TimeInMS ) override;
+	virtual uint32 GetTotalDemoTime() const override;
+	virtual bool IsDataAvailable() const override;
+	virtual void SetHighPriorityTimeRange( const uint32 StartTimeInMS, const uint32 EndTimeInMS ) override { }
+	virtual bool IsDataAvailableForTimeRange( const uint32 StartTimeInMS, const uint32 EndTimeInMS ) override { return true; }
+	virtual bool IsLoadingCheckpoint() const override { return false; }
+	virtual bool IsLive() const override;
+	virtual void DeleteFinishedStream( const FString& StreamName, const FOnDeleteFinishedStreamComplete& Delegate) const override;
+	virtual void EnumerateStreams( const FNetworkReplayVersion& ReplayVersion, const FString& UserString, const FString& MetaString, const FOnEnumerateStreamsComplete& Delegate ) override;
+	virtual void EnumerateRecentStreams( const FNetworkReplayVersion& ReplayVersion, const FString& RecentViewer, const FOnEnumerateStreamsComplete& Delegate ) override {}
+	virtual ENetworkReplayError::Type GetLastError() const override { return ENetworkReplayError::None; }
+	virtual void AddUserToReplay(const FString& UserString) override;
+	virtual void AddEvent(const uint32 TimeInMS, const FString& Group, const FString& Meta, const TArray<uint8>& Data) override;
+	virtual void EnumerateEvents(const FString& Group, const FEnumerateEventsCompleteDelegate& EnumerationCompleteDelegate) override;
+	virtual void EnumerateEvents( const FString& ReplayName, const FString& Group, const FEnumerateEventsCompleteDelegate& EnumerationCompleteDelegate ) override;
+	virtual void RequestEventData(const FString& EventID, const FOnRequestEventDataComplete& RequestEventDataComplete) override;
+	virtual void SearchEvents(const FString& EventGroup, const FOnEnumerateStreamsComplete& Delegate) override;
+	virtual void KeepReplay( const FString& ReplayName, const bool bKeep ) override;
+
+	/** FTickableObjectBase implementation */
+	virtual void Tick(float DeltaSeconds) override;
+	virtual bool IsTickable() const override { return true; }
+	virtual TStatId GetStatId() const override;
+
+	/** FTickableGameObject implementation */
+	virtual bool IsTickableWhenPaused() const override { return true; }
+
+private:
+	bool IsNamedStreamLive( const FString& StreamName ) const;
+
+	/** Handles the details of loading a checkpoint */
+	void GotoCheckpointIndexInternal(int32 CheckpointIndex, const FOnCheckpointReadyDelegate& Delegate, int32 TimeInMS);
+
+	/**
+	 * Returns a pointer to the currently active (recording or playback) replay in the owning factory's map.
+	 * May return null if if the streamer state is idle.
+	 */
+	FInMemoryReplay* GetCurrentReplay();
+
+	/**
+	 * Returns a pointer to the currently active (recording or playback) replay in the owning factory's map.
+	 * May return null if if the streamer state is idle.
+	 */
+	FInMemoryReplay* GetCurrentReplay() const;
+
+	/**
+	 * Returns a pointer to the currently active (recording or playback) replay in the owning factory's map.
+	 * Asserts if the returned replay pointer would be null.
+	 */
+	FInMemoryReplay* GetCurrentReplayChecked();
+
+	/**
+	 * Returns a pointer to the currently active (recording or playback) replay in the owning factory's map.
+	 * Asserts if the returned replay pointer would be null.
+	 */
+	FInMemoryReplay* GetCurrentReplayChecked() const;
+
+	/** Pointer to the factory that owns this streamer instance */
+	FInMemoryNetworkReplayStreamingFactory* OwningFactory;
+
+	/** Handle to the archive that will read/write the demo header */
+	TUniquePtr<FArchive> HeaderAr;
+
+	/** Handle to the archive that will read/write network packets */
+	TUniquePtr<FArchive> FileAr;
+
+	/* Handle to the archive that will read/write metadata */
+	TUniquePtr<FArchive> MetadataFileAr;
+
+	/* Handle to the archive that will read/write checkpoint files */
+	TUniquePtr<FArchive> CheckpointAr;
+
+	/** EStreamerState - Overall state of the streamer */
+	enum class EStreamerState
+	{
+		Idle,					// The streamer is idle. Either we haven't started streaming yet, or we are done
+		Recording,				// We are in the process of recording a replay to disk
+		Playback,				// We are in the process of playing a replay from disk
+	};
+
+	/** Overall state of the streamer */
+	EStreamerState StreamerState;
+
+	/** Remember the name of the current stream, if any. */
+	FString CurrentStreamName;
+};
+
+class FInMemoryNetworkReplayStreamingFactory : public INetworkReplayStreamingFactory
+{
+public:
+	virtual TSharedPtr< INetworkReplayStreamer > CreateReplayStreamer();
+
+	friend class FInMemoryNetworkReplayStreamer;
+
+private:
+	TMap<FString, TUniquePtr<FInMemoryReplay>> Replays;
+};
diff --git a/Engine/Source/Runtime/NetworkReplayStreaming/NetworkReplayStreaming/NetworkReplayStreaming.Build.cs b/Engine/Source/Runtime/NetworkReplayStreaming/NetworkReplayStreaming/NetworkReplayStreaming.Build.cs
--- a/Engine/Source/Runtime/NetworkReplayStreaming/NetworkReplayStreaming/NetworkReplayStreaming.Build.cs
+++ b/Engine/Source/Runtime/NetworkReplayStreaming/NetworkReplayStreaming/NetworkReplayStreaming.Build.cs
@@ -8,7 +8,7 @@ namespace UnrealBuildTool.Rules
 		{
 			PrivateIncludePaths.Add( "Runtime/NetworkReplayStreaming/NetworkReplayStreaming/Private" );
 
-            PrivateIncludePathModuleNames.Add("OnlineSubsystem");
+            PublicIncludePathModuleNames.Add("OnlineSubsystem");
 
 			PublicDependencyModuleNames.AddRange(
 				new string[]
