#orion, replays,  merge CL 2708382: Add more logging when events fail to upload.

diff --git a/Engine/Source/Runtime/NetworkReplayStreaming/HttpNetworkReplayStreaming/Private/HttpNetworkReplayStreaming.cpp b/Engine/Source/Runtime/NetworkReplayStreaming/HttpNetworkReplayStreaming/Private/HttpNetworkReplayStreaming.cpp
--- a/Engine/Source/Runtime/NetworkReplayStreaming/HttpNetworkReplayStreaming/Private/HttpNetworkReplayStreaming.cpp
+++ b/Engine/Source/Runtime/NetworkReplayStreaming/HttpNetworkReplayStreaming/Private/HttpNetworkReplayStreaming.cpp
@@ -1243,7 +1243,27 @@ void FHttpNetworkReplayStreamer::HttpUploadCustomEventFinished(FHttpRequestPtr H
 	}
 	else
 	{
-		UE_LOG(LogHttpReplay, Error, TEXT("FHttpNetworkReplayStreamer::HttpUploadCustomEventFinished. FAILED, Response code: %d"), HttpResponse.IsValid() ? HttpResponse->GetResponseCode() : 0);
+		FString ExtraInfo;
+
+		if ( HttpRequest.IsValid() )
+		{
+			ExtraInfo += FString::Printf( TEXT( "URL: %s" ),	*HttpRequest->GetURL() );
+			ExtraInfo += FString::Printf( TEXT( ", Verb: %s" ), *HttpRequest->GetVerb() );
+
+			const TArray< FString > AllHeaders = HttpRequest->GetAllHeaders();
+
+			for ( int i = 0; i < AllHeaders.Num(); i++ )
+			{
+				ExtraInfo += TEXT( ", " );
+				ExtraInfo += AllHeaders[i];
+			}
+		}
+		else
+		{
+			ExtraInfo = TEXT( "HttpRequest NULL." );
+		}
+
+		UE_LOG( LogHttpReplay, Error, TEXT( "FHttpNetworkReplayStreamer::HttpUploadCustomEventFinished. FAILED. Response code: %d, Extra info: %s" ), HttpResponse.IsValid() ? HttpResponse->GetResponseCode() : 0, *ExtraInfo );
 		// Don't disconect service here, just report the failure
 		//SetLastError(ENetworkReplayError::ServiceUnavailable);
 	}
