#merge CL 2711382: Properly encode event meta string when special characters are used (i.e. spaces)

diff --git a/Engine/Source/Runtime/NetworkReplayStreaming/HttpNetworkReplayStreaming/Private/HttpNetworkReplayStreaming.cpp b/Engine/Source/Runtime/NetworkReplayStreaming/HttpNetworkReplayStreaming/Private/HttpNetworkReplayStreaming.cpp
--- a/Engine/Source/Runtime/NetworkReplayStreaming/HttpNetworkReplayStreaming/Private/HttpNetworkReplayStreaming.cpp
+++ b/Engine/Source/Runtime/NetworkReplayStreaming/HttpNetworkReplayStreaming/Private/HttpNetworkReplayStreaming.cpp
@@ -682,7 +682,7 @@ bool FQueuedHttpRequestAddEvent::PreProcess( const FString& ServerURL, const FSt
 	}
 
 	// Now that we have the session name, we can set the URL
-	Request->SetURL( FString::Printf( TEXT( "%sreplay/%s/event?group=%s&time1=%i&time2=%i&meta=%s" ), *ServerURL, *SessionName, *Group, TimeInMS, TimeInMS, *Meta ) );
+	Request->SetURL( FString::Printf( TEXT( "%sreplay/%s/event?group=%s&time1=%i&time2=%i&meta=%s" ), *ServerURL, *SessionName, *Group, TimeInMS, TimeInMS, *FGenericPlatformHttp::UrlEncode( Meta ) ) );
 
 	return true;
 }
