#merge: Add ability to update existing named event

diff --git a/Engine/Source/Runtime/NetworkReplayStreaming/HttpNetworkReplayStreaming/Private/HttpNetworkReplayStreaming.cpp b/Engine/Source/Runtime/NetworkReplayStreaming/HttpNetworkReplayStreaming/Private/HttpNetworkReplayStreaming.cpp
--- a/Engine/Source/Runtime/NetworkReplayStreaming/HttpNetworkReplayStreaming/Private/HttpNetworkReplayStreaming.cpp
+++ b/Engine/Source/Runtime/NetworkReplayStreaming/HttpNetworkReplayStreaming/Private/HttpNetworkReplayStreaming.cpp
@@ -662,12 +662,13 @@ void FHttpNetworkReplayStreamer::FlushCheckpointInternal( uint32 TimeInMS )
 	AddRequestToQueue( EQueuedHttpRequestType::UploadingCheckpoint, HttpRequest );
 }
 
-FQueuedHttpRequestAddEvent::FQueuedHttpRequestAddEvent( const uint32 InTimeInMS, const FString& InGroup, const FString& InMeta, const TArray<uint8>& InData, TSharedRef< class IHttpRequest > InHttpRequest ) : FQueuedHttpRequest( EQueuedHttpRequestType::UploadingCustomEvent, InHttpRequest )
+FQueuedHttpRequestAddEvent::FQueuedHttpRequestAddEvent( const FString& InName, const uint32 InTimeInMS, const FString& InGroup, const FString& InMeta, const TArray<uint8>& InData, TSharedRef< class IHttpRequest > InHttpRequest ) : FQueuedHttpRequest( EQueuedHttpRequestType::UploadingCustomEvent, InHttpRequest )
 {
 	Request->SetVerb( TEXT( "POST" ) );
 	Request->SetHeader( TEXT( "Content-Type" ), TEXT( "application/octet-stream" ) );
 	Request->SetContent( InData );
 
+	Name		= InName;
 	TimeInMS	= InTimeInMS;
 	Group		= InGroup;
 	Meta		= InMeta;
@@ -681,22 +682,33 @@ bool FQueuedHttpRequestAddEvent::PreProcess( const FString& ServerURL, const FSt
 		return false;
 	}
 
+	//
 	// Now that we have the session name, we can set the URL
-	Request->SetURL( FString::Printf( TEXT( "%sreplay/%s/event?group=%s&time1=%i&time2=%i&meta=%s" ), *ServerURL, *SessionName, *Group, TimeInMS, TimeInMS, *FGenericPlatformHttp::UrlEncode( Meta ) ) );
+	//
+
+	if ( !Name.IsEmpty() )
+	{
+		// Update existing event
+		Request->SetURL( FString::Printf( TEXT( "%sevent/%s" ), *ServerURL, *Name ) );
+	}
+	else
+	{
+		Request->SetURL( FString::Printf( TEXT( "%sreplay/%s/event?group=%s&time1=%i&time2=%i&meta=%s" ), *ServerURL, *SessionName, *Group, TimeInMS, TimeInMS, *FGenericPlatformHttp::UrlEncode( Meta ) ) );
+	}
 
 	return true;
 }
 
-void FHttpNetworkReplayStreamer::AddEvent( const uint32 TimeInMS, const FString& Group, const FString& Meta, const TArray<uint8>& Data )
+void FHttpNetworkReplayStreamer::AddOrUpdateEvent( const FString& Name, const uint32 TimeInMS, const FString& Group, const FString& Meta, const TArray<uint8>& Data )
 {
-	if (StreamerState != EStreamerState::StreamingUp && StreamerState != EStreamerState::StreamingDown)
+	if ( StreamerState != EStreamerState::StreamingUp && StreamerState != EStreamerState::StreamingDown )
 	{
-		UE_LOG( LogHttpReplay, Warning, TEXT( "FHttpNetworkReplayStreamer::AddEvent. Not streaming." ) );
+		UE_LOG( LogHttpReplay, Warning, TEXT( "FHttpNetworkReplayStreamer::AddOrUpdateEvent. Not streaming." ) );
 		return;
 	}
 
 	// Upload a custom event to replay server
-	UE_LOG(LogHttpReplay, Verbose, TEXT("FHttpNetworkReplayStreamer::AddEvent. Size: %i, StreamChunkIndex: %i"), Data.Num(), StreamChunkIndex);
+	UE_LOG( LogHttpReplay, Verbose, TEXT( "FHttpNetworkReplayStreamer::AddEvent. Size: %i, StreamChunkIndex: %i" ), Data.Num(), StreamChunkIndex );
 
 	// Create the Http request and add to pending request list
 	TSharedRef< class IHttpRequest > HttpRequest = FHttpModule::Get().CreateRequest();
@@ -704,7 +716,18 @@ void FHttpNetworkReplayStreamer::AddEvent( const uint32 TimeInMS, const FString&
 	HttpRequest->OnProcessRequestComplete().BindRaw( this, &FHttpNetworkReplayStreamer::HttpUploadCustomEventFinished );
 
 	// Add it as a custom event so we can snag the session name at the time of send (which we should have by then)
-	AddCustomRequestToQueue( TSharedPtr< FQueuedHttpRequest >( new FQueuedHttpRequestAddEvent( TimeInMS, Group, Meta, Data, HttpRequest ) ) );
+	AddCustomRequestToQueue( TSharedPtr< FQueuedHttpRequest >( new FQueuedHttpRequestAddEvent( Name, TimeInMS, Group, Meta, Data, HttpRequest ) ) );
+}
+
+void FHttpNetworkReplayStreamer::AddEvent( const uint32 TimeInMS, const FString& Group, const FString& Meta, const TArray<uint8>& Data )
+{
+	if (StreamerState != EStreamerState::StreamingUp && StreamerState != EStreamerState::StreamingDown)
+	{
+		UE_LOG( LogHttpReplay, Warning, TEXT( "FHttpNetworkReplayStreamer::AddEvent. Not streaming." ) );
+		return;
+	}
+
+	AddOrUpdateEvent( TEXT( "" ), TimeInMS, Group, Meta, Data );
 }
 
 void FHttpNetworkReplayStreamer::DownloadHeader()
diff --git a/Engine/Source/Runtime/NetworkReplayStreaming/HttpNetworkReplayStreaming/Public/HttpNetworkReplayStreaming.h b/Engine/Source/Runtime/NetworkReplayStreaming/HttpNetworkReplayStreaming/Public/HttpNetworkReplayStreaming.h
--- a/Engine/Source/Runtime/NetworkReplayStreaming/HttpNetworkReplayStreaming/Public/HttpNetworkReplayStreaming.h
+++ b/Engine/Source/Runtime/NetworkReplayStreaming/HttpNetworkReplayStreaming/Public/HttpNetworkReplayStreaming.h
@@ -125,7 +125,7 @@ public:
 class FQueuedHttpRequestAddEvent : public FQueuedHttpRequest
 {
 public:
-	FQueuedHttpRequestAddEvent( const uint32 InTimeInMS, const FString& InGroup, const FString& InMeta, const TArray<uint8>& InData, TSharedRef< class IHttpRequest > InHttpRequest );
+	FQueuedHttpRequestAddEvent( const FString& InName, const uint32 InTimeInMS, const FString& InGroup, const FString& InMeta, const TArray<uint8>& InData, TSharedRef< class IHttpRequest > InHttpRequest );
 
 	virtual ~FQueuedHttpRequestAddEvent()
 	{
@@ -133,6 +133,7 @@ public:
 
 	virtual bool PreProcess( const FString& ServerURL, const FString& SessionName ) override;
 
+	FString		Name;
 	uint32		TimeInMS;
 	FString		Group;
 	FString		Meta;
@@ -186,7 +187,8 @@ public:
 	void SetLastError( const ENetworkReplayError::Type InLastError );
 	void CancelStreamingRequests();
 	void FlushCheckpointInternal( uint32 TimeInMS );
-	void AddEvent( const uint32 TimeInMS, const FString& Group, const FString& Meta, const TArray<uint8>& Data );
+	virtual void AddEvent( const uint32 TimeInMS, const FString& Group, const FString& Meta, const TArray<uint8>& Data ) override;
+	virtual void AddOrUpdateEvent( const FString& Name, const uint32 TimeInMS, const FString& Group, const FString& Meta, const TArray<uint8>& Data ) override;
 	void AddRequestToQueue( const EQueuedHttpRequestType::Type Type, TSharedPtr< class IHttpRequest >	Request );
 	void AddCustomRequestToQueue( TSharedPtr< FQueuedHttpRequest > Request );
 	void EnumerateCheckpoints();
diff --git a/Engine/Source/Runtime/NetworkReplayStreaming/InMemoryNetworkReplayStreaming/Public/InMemoryNetworkReplayStreaming.h b/Engine/Source/Runtime/NetworkReplayStreaming/InMemoryNetworkReplayStreaming/Public/InMemoryNetworkReplayStreaming.h
--- a/Engine/Source/Runtime/NetworkReplayStreaming/InMemoryNetworkReplayStreaming/Public/InMemoryNetworkReplayStreaming.h
+++ b/Engine/Source/Runtime/NetworkReplayStreaming/InMemoryNetworkReplayStreaming/Public/InMemoryNetworkReplayStreaming.h
@@ -70,6 +70,7 @@ public:
 	virtual ENetworkReplayError::Type GetLastError() const override { return ENetworkReplayError::None; }
 	virtual void AddUserToReplay(const FString& UserString) override;
 	virtual void AddEvent(const uint32 TimeInMS, const FString& Group, const FString& Meta, const TArray<uint8>& Data) override;
+	virtual void AddOrUpdateEvent( const FString& Name, const uint32 TimeInMS, const FString& Group, const FString& Meta, const TArray<uint8>& Data ) override {}
 	virtual void EnumerateEvents(const FString& Group, const FEnumerateEventsCompleteDelegate& EnumerationCompleteDelegate) override;
 	virtual void EnumerateEvents( const FString& ReplayName, const FString& Group, const FEnumerateEventsCompleteDelegate& EnumerationCompleteDelegate ) override;
 	virtual void RequestEventData(const FString& EventID, const FOnRequestEventDataComplete& RequestEventDataComplete) override;
diff --git a/Engine/Source/Runtime/NetworkReplayStreaming/NetworkReplayStreaming/Public/NetworkReplayStreaming.h b/Engine/Source/Runtime/NetworkReplayStreaming/NetworkReplayStreaming/Public/NetworkReplayStreaming.h
--- a/Engine/Source/Runtime/NetworkReplayStreaming/NetworkReplayStreaming/Public/NetworkReplayStreaming.h
+++ b/Engine/Source/Runtime/NetworkReplayStreaming/NetworkReplayStreaming/Public/NetworkReplayStreaming.h
@@ -176,6 +176,7 @@ public:
 	virtual bool IsDataAvailableForTimeRange( const uint32 StartTimeInMS, const uint32 EndTimeInMS ) = 0;
 	virtual bool IsLoadingCheckpoint() const = 0;
 	virtual void AddEvent( const uint32 TimeInMS, const FString& Group, const FString& Meta, const TArray<uint8>& Data ) = 0;
+	virtual void AddOrUpdateEvent( const FString& Name, const uint32 TimeInMS, const FString& Group, const FString& Meta, const TArray<uint8>& Data ) = 0;
 	virtual void EnumerateEvents( const FString& Group, const FEnumerateEventsCompleteDelegate& EnumerationCompleteDelegate ) = 0;
 	virtual void EnumerateEvents( const FString& ReplayName, const FString& Group, const FEnumerateEventsCompleteDelegate& EnumerationCompleteDelegate ) = 0;
 	virtual void RequestEventData( const FString& EventID, const FOnRequestEventDataComplete& RequestEventDataComplete ) = 0;
diff --git a/Engine/Source/Runtime/NetworkReplayStreaming/NullNetworkReplayStreaming/Public/NullNetworkReplayStreaming.h b/Engine/Source/Runtime/NetworkReplayStreaming/NullNetworkReplayStreaming/Public/NullNetworkReplayStreaming.h
--- a/Engine/Source/Runtime/NetworkReplayStreaming/NullNetworkReplayStreaming/Public/NullNetworkReplayStreaming.h
+++ b/Engine/Source/Runtime/NetworkReplayStreaming/NullNetworkReplayStreaming/Public/NullNetworkReplayStreaming.h
@@ -66,7 +66,8 @@ public:
 	virtual ENetworkReplayError::Type GetLastError() const override { return ENetworkReplayError::None; }
 	virtual void AddUserToReplay(const FString& UserString) override;
 	virtual void AddEvent(const uint32 TimeInMS, const FString& Group, const FString& Meta, const TArray<uint8>& Data) override;
-	virtual void EnumerateEvents(const FString& Group, const FEnumerateEventsCompleteDelegate& EnumerationCompleteDelegate) override;
+	virtual void AddOrUpdateEvent( const FString& Name, const uint32 TimeInMS, const FString& Group, const FString& Meta, const TArray<uint8>& Data ) override {}
+	virtual void EnumerateEvents( const FString& Group, const FEnumerateEventsCompleteDelegate& EnumerationCompleteDelegate ) override;
 	virtual void EnumerateEvents( const FString& ReplayName, const FString& Group, const FEnumerateEventsCompleteDelegate& EnumerationCompleteDelegate ) override {}
 	virtual void RequestEventData(const FString& EventID, const FOnRequestEventDataComplete& RequestEventDataComplete) override;
 	virtual void SearchEvents(const FString& EventGroup, const FOnEnumerateStreamsComplete& Delegate) override;
