Copying //UE4/Dev-Networking to Dev-Main (//UE4/Dev-Main)

==========================
MAJOR FEATURES + CHANGES
==========================

Change 2761523 on 2015/11/10 by [REDACTED_PERSON][REDACTED_HANDLE]

	Allow loading partial profiles

Change 2761650 on 2015/11/10 by [REDACTED_PERSON][REDACTED_HANDLE]

	Clamp DeltaSeconds when updating demo record time to prevent large sections of dead air during playback.

Change 2765057 on 2015/11/12 by [REDACTED_PERSON][REDACTED_HANDLE]

	Allow replay chunks to request downloads even if something else is in-flight or pending

Change 2765068 on 2015/11/12 by [REDACTED_PERSON][REDACTED_HANDLE]

	Override SComboBox::IsInteractable so that it works with analog cursor

Change 2765895 on 2015/11/13 by [REDACTED_PERSON][REDACTED_HANDLE]

	Fix issue with FOnlineAsyncGetUserAccessCode never completing if the call to getUserAccessCode fails

Change 2765904 on 2015/11/13 by [REDACTED_PERSON][REDACTED_HANDLE]

	Set Windows console window width to match screen buffer width. Eliminates the need for manual horizontal resizing or scrolling of log windows, handy for dedicated servers running with -log.

diff --git a/Engine/Source/Programs/NetworkProfiler/NetworkProfiler/MainWindow.cs b/Engine/Source/Programs/NetworkProfiler/NetworkProfiler/MainWindow.cs
--- a/Engine/Source/Programs/NetworkProfiler/NetworkProfiler/MainWindow.cs
+++ b/Engine/Source/Programs/NetworkProfiler/NetworkProfiler/MainWindow.cs
@@ -81,6 +81,9 @@ namespace NetworkProfiler
 
 		PartialNetworkStream CurrentStreamSelection = null;
 
+		/** If non 0, we will early out of loading in this many minutes worth of profile time */
+		int MaxProfileMinutes		= 0;
+
 		public MainWindow()
 		{
 			InitializeComponent();
@@ -766,6 +769,28 @@ namespace NetworkProfiler
 		{
 			UpdateConnectionFilter();
 		}
+
+		public int GetMaxProfileMinutes()
+		{
+			return MaxProfileMinutes;
+		}
+
+		private void MaxProfileMinutesTextBox_TextChanged( object sender, EventArgs e )
+		{
+			if ( MaxProfileMinutesTextBox.Text == "" )
+			{
+				MaxProfileMinutes = 0;
+			}
+			else try
+			{
+				MaxProfileMinutes = Int32.Parse( MaxProfileMinutesTextBox.Text );
+			}
+			catch
+			{
+				MaxProfileMinutes = 0;
+				MaxProfileMinutesTextBox.Text = "";
+			}
+		}
 	}
 
 	public class FilterValues
diff --git a/Engine/Source/Programs/NetworkProfiler/NetworkProfiler/StreamParser.cs b/Engine/Source/Programs/NetworkProfiler/NetworkProfiler/StreamParser.cs
--- a/Engine/Source/Programs/NetworkProfiler/NetworkProfiler/StreamParser.cs
+++ b/Engine/Source/Programs/NetworkProfiler/NetworkProfiler/StreamParser.cs
@@ -111,6 +111,8 @@ namespace NetworkProfiler
 
 			var AllFrames = new PartialNetworkStream( NetworkStream.NameIndexUnreal, 1.0f / 30.0f );
 
+			int EarlyOutMinutes = InMainWindow.GetMaxProfileMinutes();
+
 			// Parse stream till we reach the end, marked by special token.
 			bool bHasReachedEndOfStream = false;
 			while( bHasReachedEndOfStream == false )
@@ -211,6 +213,11 @@ namespace NetworkProfiler
 						CurrentFrameTokens.Add(Token);
 					}
 				}
+
+				if ( EarlyOutMinutes > 0 && ( ( AllFrames.EndTime - AllFrames.StartTime  ) > 60 * EarlyOutMinutes ) )
+				{
+					break;
+				}
 			}
 
 			InMainWindow.SetCurrentStreamSelection( NetworkStream, AllFrames, false );
diff --git a/Engine/Source/Runtime/Core/Private/Windows/WindowsPlatformOutputDevices.cpp b/Engine/Source/Runtime/Core/Private/Windows/WindowsPlatformOutputDevices.cpp
--- a/Engine/Source/Runtime/Core/Private/Windows/WindowsPlatformOutputDevices.cpp
+++ b/Engine/Source/Runtime/Core/Private/Windows/WindowsPlatformOutputDevices.cpp
@@ -218,6 +218,16 @@ void FOutputDeviceConsoleWindows::Show( bool ShowWindow )
 
 				SetConsoleScreenBufferSize( ConsoleHandle, Size );
 
+				CONSOLE_SCREEN_BUFFER_INFO ConsoleInfo;
+
+				// Try to set the window width to match the screen buffer width, so that no manual horizontal scrolling or resizing is necessary
+				if (::GetConsoleScreenBufferInfo( ConsoleHandle, &ConsoleInfo ) != 0)
+				{
+					SMALL_RECT NewConsoleWindowRect = ConsoleInfo.srWindow;
+					NewConsoleWindowRect.Right = ConsoleInfo.dwSize.X - 1;
+					::SetConsoleWindowInfo( ConsoleHandle, true, &NewConsoleWindowRect );
+				}
+
 				RECT WindowRect;
 				::GetWindowRect( GetConsoleWindow(), &WindowRect );
 
diff --git a/Engine/Source/Runtime/Engine/Private/DemoNetDriver.cpp b/Engine/Source/Runtime/Engine/Private/DemoNetDriver.cpp
--- a/Engine/Source/Runtime/Engine/Private/DemoNetDriver.cpp
+++ b/Engine/Source/Runtime/Engine/Private/DemoNetDriver.cpp
@@ -1277,7 +1277,12 @@ void UDemoNetDriver::TickDemoRecord( float DeltaSeconds )
 		return;
 	}
 
-	DemoCurrentTime += DeltaSeconds;
+	{
+		// Since the DeltaSeconds here hasn't been clamped or dilated by the engine,
+		// clamp here to prevent large gaps in replay frame times that can cause a lot of dead air when played back.
+		const float ClampedDeltaSeconds = FMath::Clamp( DeltaSeconds, 1.0f / 1000.0f, 1.0f / 2.0f );
+		DemoCurrentTime += ClampedDeltaSeconds;
+	}
 
 	ReplayStreamer->UpdateTotalDemoTime( GetDemoCurrentTimeInMS() );
 
diff --git a/Engine/Source/Runtime/NetworkReplayStreaming/HttpNetworkReplayStreaming/Private/HttpNetworkReplayStreaming.cpp b/Engine/Source/Runtime/NetworkReplayStreaming/HttpNetworkReplayStreaming/Private/HttpNetworkReplayStreaming.cpp
--- a/Engine/Source/Runtime/NetworkReplayStreaming/HttpNetworkReplayStreaming/Private/HttpNetworkReplayStreaming.cpp
+++ b/Engine/Source/Runtime/NetworkReplayStreaming/HttpNetworkReplayStreaming/Private/HttpNetworkReplayStreaming.cpp
@@ -306,14 +306,14 @@ void FHttpNetworkReplayStreamer::AddRequestToQueue( const EQueuedHttpRequestType
 {
 	UE_LOG( LogHttpReplay, Verbose, TEXT( "FHttpNetworkReplayStreamer::AddRequestToQueue. Type: %s" ), EQueuedHttpRequestType::ToString( Type ) );
 
-	QueuedHttpRequests.Enqueue( TSharedPtr< FQueuedHttpRequest >( new FQueuedHttpRequest( Type, Request ) ) );
+	QueuedHttpRequests.Add( TSharedPtr< FQueuedHttpRequest >( new FQueuedHttpRequest( Type, Request ) ) );
 }
 
 void FHttpNetworkReplayStreamer::AddCustomRequestToQueue( TSharedPtr< FQueuedHttpRequest > Request )
 {
 	UE_LOG( LogHttpReplay, Verbose, TEXT( "FHttpNetworkReplayStreamer::AddCustomRequestToQueue. Type: %s" ), EQueuedHttpRequestType::ToString( Request->Type ) );
 
-	QueuedHttpRequests.Enqueue( Request );
+	QueuedHttpRequests.Add( Request );
 }
 
 void FHttpNetworkReplayStreamer::StopStreaming()
@@ -445,6 +445,11 @@ void FHttpNetworkReplayStreamer::FlushStream()
 
 void FHttpNetworkReplayStreamer::ConditionallyFlushStream()
 {
+	if ( IsHttpRequestInFlight() || HasPendingHttpRequests() )
+	{
+		return;
+	}
+	
 	const double FLUSH_TIME_IN_SECONDS = 10;
 
 	if ( FPlatformTime::Seconds() - LastChunkTime > FLUSH_TIME_IN_SECONDS )
@@ -744,17 +749,38 @@ void FHttpNetworkReplayStreamer::DownloadHeader()
 	AddRequestToQueue( EQueuedHttpRequestType::DownloadingHeader, HttpRequest );
 }
 
+bool FHttpNetworkReplayStreamer::IsTaskPendingOrInFlight( const EQueuedHttpRequestType::Type Type ) const
+{
+	for ( const TSharedPtr< FQueuedHttpRequest >& Request : QueuedHttpRequests )
+	{
+		if ( Request->Type == Type )
+		{
+			return true;
+		}
+	}
+
+	if ( InFlightHttpRequest.IsValid() )
+	{
+		if ( InFlightHttpRequest->Type == Type )
+		{
+			return true;
+		}
+	}
+
+	return false;
+}
+
 void FHttpNetworkReplayStreamer::ConditionallyDownloadNextChunk()
 {
-	if ( IsHttpRequestInFlight() )
+	if ( GotoCheckpointDelegate.IsBound() )
 	{
-		// We never download next chunk if there is an http request in flight
+		// Don't download stream chunk while we're waiting for a checkpoint to download
 		return;
 	}
 
-	if ( GotoCheckpointDelegate.IsBound() )
+	if ( IsTaskPendingOrInFlight( EQueuedHttpRequestType::DownloadingStream ) )
 	{
-		// Don't download stream while we're waiting for a checkpoint to download
+		// Only download one chunk at a time
 		return;
 	}
 
@@ -859,6 +885,11 @@ void FHttpNetworkReplayStreamer::RefreshViewer( const bool bFinal )
 
 void FHttpNetworkReplayStreamer::ConditionallyRefreshViewer()
 {
+	if ( IsHttpRequestInFlight() || HasPendingHttpRequests() )
+	{
+		return;
+	}
+
 	const double REFRESH_VIEWER_IN_SECONDS = 10;
 
 	if ( FPlatformTime::Seconds() - LastRefreshViewerTime > REFRESH_VIEWER_IN_SECONDS )
@@ -884,10 +915,7 @@ void FHttpNetworkReplayStreamer::CancelStreamingRequests()
 	}
 
 	// Empty the request queue
-	TSharedPtr< FQueuedHttpRequest > QueuedRequest;
-	while ( QueuedHttpRequests.Dequeue( QueuedRequest ) )
-	{
-	}
+	QueuedHttpRequests.Empty();
 
 	StreamerState			= EStreamerState::Idle;
 	bStopStreamingCalled	= false;
@@ -1106,6 +1134,11 @@ void FHttpNetworkReplayStreamer::EnumerateCheckpoints()
 
 void FHttpNetworkReplayStreamer::ConditionallyEnumerateCheckpoints()
 {
+	if ( IsHttpRequestInFlight() || HasPendingHttpRequests() )
+	{
+		return;
+	}
+
 	if ( !bStreamIsLive )
 	{
 		// We don't need to enumerate more than once for non live streams
@@ -1714,10 +1747,12 @@ bool FHttpNetworkReplayStreamer::ProcessNextHttpRequest()
 		return false;
 	}
 
-	TSharedPtr< FQueuedHttpRequest > QueuedRequest;
-
-	if ( QueuedHttpRequests.Dequeue( QueuedRequest ) )
+	if ( QueuedHttpRequests.Num() > 0 )
 	{
+		TSharedPtr< FQueuedHttpRequest > QueuedRequest = QueuedHttpRequests[0];
+
+		QueuedHttpRequests.RemoveAt( 0 );
+
 		UE_LOG( LogHttpReplay, Verbose, TEXT( "FHttpNetworkReplayStreamer::ProcessNextHttpRequest. Dequeue Type: %s" ), EQueuedHttpRequestType::ToString( QueuedRequest->Type ) );
 
 		check( !InFlightHttpRequest.IsValid() );
@@ -1761,22 +1796,16 @@ void FHttpNetworkReplayStreamer::ProcessRequestInternal( TSharedPtr< class IHttp
 
 void FHttpNetworkReplayStreamer::Tick( const float DeltaTime )
 {
-	if ( IsHttpRequestInFlight() )
-	{
-		// We only process one http request at a time to keep things simple
-		return;
-	}
-
 	// Attempt to process the next http request
 	if ( ProcessNextHttpRequest() )
 	{
-		// We can now return since we know there is a request in flight
 		check( IsHttpRequestInFlight() );
-		return;
 	}
 
-	// We should have no pending or requests in flight at this point
-	check( !HasPendingHttpRequests() );
+	if ( bStopStreamingCalled )
+	{
+		return;
+	}
 
 	if ( StreamerState == EStreamerState::StreamingUp )
 	{
@@ -1784,6 +1813,12 @@ void FHttpNetworkReplayStreamer::Tick( const float DeltaTime )
 	}
 	else if ( StreamerState == EStreamerState::StreamingDown )
 	{
+		if ( IsTaskPendingOrInFlight( EQueuedHttpRequestType::StartDownloading ) )
+		{
+			// If we're still waiting on finalizing the download request then return
+			return;
+		}
+
 		// Check to see if we're done downloading the high priority portion of the strema
 		// If so, we can cancel the request
 		if ( HighPriorityEndTime > 0 && StreamTimeRangeEnd >= HighPriorityEndTime )
@@ -1812,7 +1847,7 @@ bool FHttpNetworkReplayStreamer::IsHttpRequestInFlight() const
 bool FHttpNetworkReplayStreamer::HasPendingHttpRequests() const
 {
 	// If there is currently one in flight, or we have more to process, return true
-	return IsHttpRequestInFlight() || !QueuedHttpRequests.IsEmpty();
+	return IsHttpRequestInFlight() || QueuedHttpRequests.Num() > 0;
 }
 
 bool FHttpNetworkReplayStreamer::IsStreaming() const
diff --git a/Engine/Source/Runtime/NetworkReplayStreaming/HttpNetworkReplayStreaming/Public/HttpNetworkReplayStreaming.h b/Engine/Source/Runtime/NetworkReplayStreaming/HttpNetworkReplayStreaming/Public/HttpNetworkReplayStreaming.h
--- a/Engine/Source/Runtime/NetworkReplayStreaming/HttpNetworkReplayStreaming/Public/HttpNetworkReplayStreaming.h
+++ b/Engine/Source/Runtime/NetworkReplayStreaming/HttpNetworkReplayStreaming/Public/HttpNetworkReplayStreaming.h
@@ -182,6 +182,7 @@ public:
 	void ConditionallyFlushStream();
 	void StopUploading();
 	void DownloadHeader();
+	bool IsTaskPendingOrInFlight( const EQueuedHttpRequestType::Type Type ) const;
 	void ConditionallyDownloadNextChunk();
 	void RefreshViewer( const bool bFinal );
 	void ConditionallyRefreshViewer();
@@ -262,7 +263,7 @@ public:
 
 	FReplayEventList					CheckpointList;
 
-	TQueue< TSharedPtr< FQueuedHttpRequest > >	QueuedHttpRequests;
+	TArray< TSharedPtr< FQueuedHttpRequest > >	QueuedHttpRequests;
 	TSharedPtr< FQueuedHttpRequest >			InFlightHttpRequest;
 };
 
diff --git a/Engine/Source/Runtime/Slate/Public/Widgets/Input/SComboBox.h b/Engine/Source/Runtime/Slate/Public/Widgets/Input/SComboBox.h
--- a/Engine/Source/Runtime/Slate/Public/Widgets/Input/SComboBox.h
+++ b/Engine/Source/Runtime/Slate/Public/Widgets/Input/SComboBox.h
@@ -281,6 +281,11 @@ protected:
 		return true;
 	}
 
+	virtual bool IsInteractable() const
+	{
+		return IsEnabled() && SupportsKeyboardFocus();
+	}
+
 private:
 
 	/** Generate a row for the InItem in the combo box's list (passed in as OwnerTable). Do this by calling the user-specified OnGenerateWidget */
