Copying //UE4/Orion-Staging to //UE4/Main (//Orion/Dev-General @ 2855324)

diff --git a/Engine/Source/Runtime/Engine/Classes/Materials/MaterialParameterCollection.h b/Engine/Source/Runtime/Engine/Classes/Materials/MaterialParameterCollection.h
--- a/Engine/Source/Runtime/Engine/Classes/Materials/MaterialParameterCollection.h
+++ b/Engine/Source/Runtime/Engine/Classes/Materials/MaterialParameterCollection.h
@@ -32,7 +32,12 @@ USTRUCT()
 struct FCollectionScalarParameter : public FCollectionParameterBase
 {
 	GENERATED_USTRUCT_BODY()
-	
+
+	FCollectionScalarParameter()
+	{
+		ParameterName = FName(TEXT("Scalar"));
+	}
+
 	UPROPERTY(EditAnywhere, Category=Parameter)
 	float DefaultValue;
 };
@@ -43,6 +48,11 @@ struct FCollectionVectorParameter : public FCollectionParameterBase
 {
 	GENERATED_USTRUCT_BODY()
 	
+	FCollectionVectorParameter()
+	{
+		ParameterName = FName(TEXT("Vector"));
+	}
+
 	UPROPERTY(EditAnywhere, Category=Parameter)
 	FLinearColor DefaultValue;
 };
@@ -70,7 +80,6 @@ class UMaterialParameterCollection : public UObject
 #if WITH_EDITOR
 	virtual void PreEditChange(class FEditPropertyChain& PropertyAboutToChange) override;
 	virtual void PostEditChangeProperty(FPropertyChangedEvent& PropertyChangedEvent) override;
-	virtual void PostEditChangeChainProperty(struct FPropertyChangedChainEvent& PropertyChangedEvent) override;
 #endif // WITH_EDITOR
 	virtual void PostLoad() override;
 	//~ End UObject Interface.
diff --git a/Engine/Source/Runtime/Engine/Classes/Particles/SubUVAnimation.h b/Engine/Source/Runtime/Engine/Classes/Particles/SubUVAnimation.h
--- a/Engine/Source/Runtime/Engine/Classes/Particles/SubUVAnimation.h
+++ b/Engine/Source/Runtime/Engine/Classes/Particles/SubUVAnimation.h
@@ -151,6 +151,10 @@ public:
 	virtual void BeginDestroy() override;
 	virtual bool IsReadyForFinishDestroy() override;
 	virtual void FinishDestroy() override;
+	virtual bool NeedsLoadForServer() const override
+	{
+		return false;
+	}
 	//~ End UObject Interface.
 
 private:
diff --git a/Engine/Source/Runtime/Engine/Private/Animation/AnimSequence.cpp b/Engine/Source/Runtime/Engine/Private/Animation/AnimSequence.cpp
--- a/Engine/Source/Runtime/Engine/Private/Animation/AnimSequence.cpp
+++ b/Engine/Source/Runtime/Engine/Private/Animation/AnimSequence.cpp
@@ -433,7 +433,16 @@ bool UAnimSequence::IsValidToPlay() const
 
 void UAnimSequence::SortSyncMarkers()
 {
+	// First make sure all SyncMarkers are within a valid range
+	for (auto& SyncMarker : AuthoredSyncMarkers)
+	{
+		SyncMarker.Time = FMath::Clamp(SyncMarker.Time, 0.f, SequenceLength);
+	}
+
+	// Then sort
 	AuthoredSyncMarkers.Sort();
+
+	// Then refresh data
 	RefreshSyncMarkerDataFromAuthored();
 }
 
diff --git a/Engine/Source/Runtime/Engine/Private/ChartCreation.cpp b/Engine/Source/Runtime/Engine/Private/ChartCreation.cpp
--- a/Engine/Source/Runtime/Engine/Private/ChartCreation.cpp
+++ b/Engine/Source/Runtime/Engine/Private/ChartCreation.cpp
@@ -1154,7 +1154,7 @@ void UEngine::DumpFPSChartToAnalyticsParams(float TotalTime, float DeltaTime, in
 	InParamArray.Add(FAnalyticsEventAttribute(TEXT("BuildType"), EBuildConfigurations::ToString(FApp::GetBuildConfiguration())));
 	InParamArray.Add(FAnalyticsEventAttribute(TEXT("DateStamp"), FDateTime::Now().ToString()));
 
-	InParamArray.Add(FAnalyticsEventAttribute(TEXT("Platform"), FString::Printf(TEXT("%s"), FPlatformProperties::PlatformName())));
+	InParamArray.Add(FAnalyticsEventAttribute(TEXT("Platform"), FString::Printf(TEXT("%s"), ANSI_TO_TCHAR(FPlatformProperties::PlatformName()))));
 	InParamArray.Add(FAnalyticsEventAttribute(TEXT("OS"), FString::Printf(TEXT("%s %s"), *OSMajor, *OSMinor)));
 	InParamArray.Add(FAnalyticsEventAttribute(TEXT("CPU"), FString::Printf(TEXT("%s %s"), *FPlatformMisc::GetCPUVendor().Trim().TrimTrailing(), *FPlatformMisc::GetCPUBrand().Trim().TrimTrailing())));
 	InParamArray.Add(FAnalyticsEventAttribute(TEXT("GPU"), FPlatformMisc::GetPrimaryGPUBrand().Trim().TrimTrailing()));
diff --git a/Engine/Source/Runtime/Engine/Private/Components/AudioComponent.cpp b/Engine/Source/Runtime/Engine/Private/Components/AudioComponent.cpp
--- a/Engine/Source/Runtime/Engine/Private/Components/AudioComponent.cpp
+++ b/Engine/Source/Runtime/Engine/Private/Components/AudioComponent.cpp
@@ -167,6 +167,8 @@ void UAudioComponent::PlayInternal(const float StartTime, const float FadeInDura
 	// Bump ActiveCount... this is used to determine if an audio component is still active after "finishing"
 	++ActiveCount;
 
+	// Whether or not we managed to actually try to play the sound
+	bool bPlayedSound = false;
 	if (Sound && (World == nullptr || World->bAllowAudioPlayback))
 	{
 		if (FAudioDevice* AudioDevice = GetAudioDevice())
@@ -246,11 +248,18 @@ void UAudioComponent::PlayInternal(const float StartTime, const float FadeInDura
 
 				// TODO - Audio Threading. This call would be a task call to dispatch to the audio thread
 				AudioDevice->AddNewActiveSound(NewActiveSound);
-
 				bIsActive = true;
+				bPlayedSound = true;
 			}
 		}
 	}
+
+	// Failed to actually try and play the sound so notify that the sound is complete.
+	// This will destroy the audio component if this is the last active sound on the audio component
+	if (!bPlayedSound)
+	{
+		PlaybackCompleted(true);
+	}
 }
 
 FAudioDevice* UAudioComponent::GetAudioDevice() const
diff --git a/Engine/Source/Runtime/Engine/Private/Components/SkeletalMeshComponent.cpp b/Engine/Source/Runtime/Engine/Private/Components/SkeletalMeshComponent.cpp
--- a/Engine/Source/Runtime/Engine/Private/Components/SkeletalMeshComponent.cpp
+++ b/Engine/Source/Runtime/Engine/Private/Components/SkeletalMeshComponent.cpp
@@ -1184,7 +1184,7 @@ void USkeletalMeshComponent::RefreshBoneTransforms(FActorComponentTickFunction*
 	const bool bInvalidCachedCurve = bDoEvaluationRateOptimization && 
 									CachedCurve.Num() != GetCurveNumber(SkeletalMesh->Skeleton);
 
-	const bool bShouldDoEvaluation = !bDoEvaluationRateOptimization || bInvalidCachedBones || !AnimUpdateRateParams->ShouldSkipEvaluation();
+	const bool bShouldDoEvaluation = !bDoEvaluationRateOptimization || bInvalidCachedBones || bInvalidCachedCurve || !AnimUpdateRateParams->ShouldSkipEvaluation();
 
 	const bool bDoPAE = !!CVarUseParallelAnimationEvaluation.GetValueOnGameThread() && FApp::ShouldUseThreadingForPerformance();
 
diff --git a/Engine/Source/Runtime/Engine/Private/ContentStreaming.cpp b/Engine/Source/Runtime/Engine/Private/ContentStreaming.cpp
--- a/Engine/Source/Runtime/Engine/Private/ContentStreaming.cpp
+++ b/Engine/Source/Runtime/Engine/Private/ContentStreaming.cpp
@@ -277,15 +277,18 @@ static const TCHAR* GStreamTypeNames[] =
 
 struct FTexturePriority
 {
-	FTexturePriority( float InPriority, int32 InTextureIndex )
+	FTexturePriority( float InPriority, int32 InTextureIndex, const UTexture2D* InTexture)
 	:	Priority( InPriority )
 	,	TextureIndex( InTextureIndex )
+	,	Texture( InTexture )
 	{
 	}
 	/** Texture priority, higher value means more important. */
 	float Priority;
 	/** Index into the FStreamingManagerTexture::StreamingTextures array. */
 	int32 TextureIndex;
+	/** The texture to stream. Only used for validation. */
+	const UTexture2D* Texture;
 };
 
 /** Self-contained structure to manage a streaming texture, possibly on a separate thread. */
@@ -1239,7 +1242,11 @@ private:
 				// Add to sort list, if it wants to stream in or could potentially stream out.
 				if ( StreamingTexture.WantedMips > StreamingTexture.ResidentMips || StreamingTexture.ResidentMips > StreamingTexture.MinAllowedMips )
 				{
-					FTexturePriority* TexturePriority = new (PrioritizedTextures) FTexturePriority( StreamingTexture.CalcPriority(), Index );
+					const UTexture2D* Texture = StreamingTexture.Texture; // This needs to go on the stack in case the main thread updates it at the same time (for the condition consistency).
+					if (Texture)
+					{
+						FTexturePriority* TexturePriority = new (PrioritizedTextures) FTexturePriority( StreamingTexture.CalcPriority(), Index, Texture );
+					}
 				}
 
 				// Accumulate streaming numbers.
@@ -3547,12 +3554,13 @@ void FStreamingManagerTexture::StreamTextures( bool bProcessEverything )
 		while ( HighPrioIndex <= LowPrioIndex && TempMemoryUsed < MaxTempMemoryUsed )
 		{
 			const FTexturePriority& TexturePriority = PrioritizedTextures[ HighPrioIndex ];
-			FStreamingTexture& HighPrioTexture = StreamingTextures[ TexturePriority.TextureIndex ];
-			bool bStreamInTexture = false;
 
-			// Check that texture hasn't been marked for removal.
-			if ( HighPrioTexture.Texture != NULL )
+			// Because deleted textures can shrink StreamingTextures, we need to check that the index is in range.
+			// Even if it is, because of the RemoveSwap logic, it could refer to the wrong texture now.
+			if (TexturePriority.TextureIndex < StreamingTextures.Num() && TexturePriority.Texture == StreamingTextures[TexturePriority.TextureIndex].Texture)
 			{
+				FStreamingTexture& HighPrioTexture = StreamingTextures[TexturePriority.TextureIndex];
+
 				// Do we want to cancel unload? Do we want more mips than requested?
 				if ( HighPrioTexture.bInFlight && HighPrioTexture.RequestedMips < HighPrioTexture.ResidentMips && HighPrioTexture.RequestedMips < HighPrioTexture.WantedMips )
 				{
diff --git a/Engine/Source/Runtime/Engine/Private/DataChannel.cpp b/Engine/Source/Runtime/Engine/Private/DataChannel.cpp
--- a/Engine/Source/Runtime/Engine/Private/DataChannel.cpp
+++ b/Engine/Source/Runtime/Engine/Private/DataChannel.cpp
@@ -1515,6 +1515,12 @@ void UActorChannel::DestroyActorAndComponents()
 		if ( CreateSubObjects[i].IsValid() )
 		{
 			UObject *SubObject = CreateSubObjects[i].Get();
+
+			if ( Connection != nullptr && Connection->Driver != nullptr )
+			{
+				Connection->Driver->RepChangedPropertyTrackerMap.Remove( SubObject );
+			}
+
 			Actor->OnSubobjectDestroyFromReplication(SubObject);
 			SubObject->PreDestroyFromReplication();
 			SubObject->MarkPendingKill();
@@ -2615,6 +2621,11 @@ UObject* UActorChannel::ReadContentBlockHeader(FInBunch & Bunch, bool& bObjectDe
 			// Stop tracking this sub-object
 			CreateSubObjects.Remove( SubObj );
 
+			if ( Connection != nullptr && Connection->Driver != nullptr )
+			{
+				Connection->Driver->RepChangedPropertyTrackerMap.Remove( SubObj );
+			}
+
 			Actor->OnSubobjectDestroyFromReplication( SubObj );
 
 			SubObj->PreDestroyFromReplication();
diff --git a/Engine/Source/Runtime/Engine/Private/DataReplication.cpp b/Engine/Source/Runtime/Engine/Private/DataReplication.cpp
--- a/Engine/Source/Runtime/Engine/Private/DataReplication.cpp
+++ b/Engine/Source/Runtime/Engine/Private/DataReplication.cpp
@@ -704,13 +704,10 @@ bool FObjectReplicator::ReceivedBunch( FInBunch& Bunch, const FReplicationFlags&
 			}
 
 			// Destroy the parameters.
-			//warning: highly dependent on UObject::ProcessEvent freeing of parms!
-			for ( UProperty * Destruct=Function->DestructorLink; Destruct; Destruct=Destruct->DestructorLinkNext )
+			// warning: highly dependent on UObject::ProcessEvent freeing of parms!
+			for( TFieldIterator<UProperty> It(Function); It && (It->PropertyFlags & (CPF_Parm|CPF_ReturnParm))==CPF_Parm; ++It )
 			{
-				if( Destruct->IsInContainer(Function->ParmsSize) )
-				{
-					Destruct->DestroyValue_InContainer(Parms);
-				}
+				It->DestroyValue_InContainer(Parms);
 			}
 
 			Mark.Pop();
diff --git a/Engine/Source/Runtime/Engine/Private/DemoNetDriver.cpp b/Engine/Source/Runtime/Engine/Private/DemoNetDriver.cpp
--- a/Engine/Source/Runtime/Engine/Private/DemoNetDriver.cpp
+++ b/Engine/Source/Runtime/Engine/Private/DemoNetDriver.cpp
@@ -305,6 +305,8 @@ bool UDemoNetDriver::InitBase( bool bInitAsClient, FNetworkNotify* InNotify, con
 		bIsFastForwarding				= false;
 		bIsFastForwardingForCheckpoint	= false;
 		bWasStartStreamingSuccessful	= true;
+		SavedReplicatedWorldTimeSeconds	= 0.0f;
+		SavedSecondsToSkip				= 0.0f;
 
 		ResetDemoState();
 
@@ -1172,6 +1174,21 @@ void UDemoNetDriver::SaveCheckpoint()
 
 	const double StartCheckpointTime = FPlatformTime::Seconds();
 
+	// Save the replicated server time so we can restore it after the checkpoint has been serialized.
+	// This preserves the existing behavior and prevents clients from receiving updated server time
+	// more often than the normal update rate.
+	float SavedReplicatedServerTimeSeconds = -1.0f;
+
+	AGameState* const GameState = World != nullptr ? World->GetGameState() : nullptr;
+
+	// Normally AGameState::ReplicatedWorldTimeSeconds is only updated periodically,
+	// but we want to make sure it's accurate for the checkpoint.
+	if (GameState != nullptr)
+	{
+		SavedReplicatedServerTimeSeconds = GameState->ReplicatedWorldTimeSeconds;
+		GameState->UpdateServerTimeSeconds();
+	}
+
 	// First, save the current guid cache
 	SerializeGuidCache( GuidCache, CheckpointArchive );
 
@@ -1245,6 +1262,12 @@ void UDemoNetDriver::SaveCheckpoint()
 		ReplayStreamer->FlushCheckpoint( SavedAbsTimeMS );
 	}
 
+	// Restore the game state's replicated world time
+	if (GameState != nullptr)
+	{
+		GameState->ReplicatedWorldTimeSeconds = SavedReplicatedServerTimeSeconds;
+	}
+
 	const double EndCheckpointTime = FPlatformTime::Seconds();
 
 	const float CheckpointTimeInMS = ( EndCheckpointTime - StartCheckpointTime ) * 1000.0f;
@@ -1655,6 +1678,7 @@ void UDemoNetDriver::SkipTimeInternal( const float SecondsToSkip, const bool InF
 	check( !bIsFastForwarding );				// Can only do one of these at a time (use tasks to gate this)
 	check( !bIsFastForwardingForCheckpoint );	// Can only do one of these at a time (use tasks to gate this)
 
+	SavedSecondsToSkip = SecondsToSkip;
 	DemoCurrentTime += SecondsToSkip;
 
 	bIsFastForwarding				= InFastForward;
@@ -1775,12 +1799,9 @@ void UDemoNetDriver::TickDemoPlayback( float DeltaSeconds )
 
 void UDemoNetDriver::FinalizeFastForward( const float StartTime )
 {
+	// This must be set before we CallRepNotifies or they might be skipped again
 	bIsFastForwarding = false;
 
-	// We may have been fast-forwarding immediately after loading a checkpoint
-	// for fine-grained scrubbing. If so, at this point we are no longer loading a checkpoint.
-	bIsFastForwardingForCheckpoint = false;
-
 	// Flush all pending RepNotifies that were built up during the fast-forward.
 	if ( ServerConnection != nullptr )
 	{
@@ -1796,6 +1817,28 @@ void UDemoNetDriver::FinalizeFastForward( const float StartTime )
 		}
 	}
 
+	AGameState* const GameState = World != nullptr ? World->GetGameState() : nullptr;
+
+	// Correct server world time for fast-forwarding after a checkpoint
+	if (GameState != nullptr)
+	{
+		if (bIsFastForwardingForCheckpoint)
+		{
+			const float PostCheckpointServerTime = SavedReplicatedWorldTimeSeconds + SavedSecondsToSkip;
+			GameState->ReplicatedWorldTimeSeconds = PostCheckpointServerTime;
+		}
+
+		// Correct client world time for any fast-forward, checkpoint or not
+		World->TimeSeconds += SavedSecondsToSkip;
+
+		// Correct the ServerWorldTimeSecondsDelta
+		GameState->OnRep_ReplicatedWorldTimeSeconds();
+	}
+
+	// We may have been fast-forwarding immediately after loading a checkpoint
+	// for fine-grained scrubbing. If so, at this point we are no longer loading a checkpoint.
+	bIsFastForwardingForCheckpoint = false;
+
 	// Reset the never-queue GUID list, we'll rebuild it
 	NonQueuedGUIDsForScrubbing.Reset();
 
@@ -2114,6 +2157,16 @@ void UDemoNetDriver::LoadCheckpoint( FArchive* GotoCheckpointArchive, int64 Goto
 
 	bDemoPlaybackDone = false;
 
+	// Save the replicated server time here
+	if (World != nullptr)
+	{
+		const AGameState* const GameState = World->GetGameState();
+		if (GameState != nullptr)
+		{
+			SavedReplicatedWorldTimeSeconds = GameState->ReplicatedWorldTimeSeconds;
+		}
+	}
+
 	if ( SpectatorController && ViewTargetGUID.IsValid() )
 	{
 		AActor* ViewTarget = Cast< AActor >( GuidCache->GetObjectFromNetGUID( ViewTargetGUID, false ) );
diff --git a/Engine/Source/Runtime/Engine/Private/GameEngine.cpp b/Engine/Source/Runtime/Engine/Private/GameEngine.cpp
--- a/Engine/Source/Runtime/Engine/Private/GameEngine.cpp
+++ b/Engine/Source/Runtime/Engine/Private/GameEngine.cpp
@@ -48,6 +48,17 @@
 
 ENGINE_API bool GDisallowNetworkTravel = false;
 
+// How slow must a frame be (in seconds) to be logged out (<= 0 to disable)
+ENGINE_API float GSlowFrameLoggingThreshold = 0.0f;
+
+static FAutoConsoleVariableRef CvarSlowFrameLoggingThreshold(
+	TEXT("t.SlowFrameLoggingThreshold"),
+	GSlowFrameLoggingThreshold,
+	TEXT("How slow must a frame be (in seconds) to be logged out (<= 0 to disable)."),
+	ECVF_Default
+	);
+
+
 /** Benchmark results to the log */
 static void RunSynthBenchmark(const TArray<FString>& Args)
 {
@@ -915,6 +926,11 @@ void UGameEngine::Tick( float DeltaSeconds, bool bIdleMode )
 #endif
 	}
 
+	if ((GSlowFrameLoggingThreshold > 0.0f) && (DeltaSeconds > GSlowFrameLoggingThreshold))
+	{
+		UE_LOG(LogEngine, Log, TEXT("Slow GT frame detected (GT frame %u, delta time %f s)"), GFrameCounter - 1, DeltaSeconds);
+	}
+
 	// Tick the module manager
 	IHotReloadInterface* HotReload = IHotReloadInterface::GetPtr();
 	if(HotReload != nullptr)
diff --git a/Engine/Source/Runtime/Engine/Private/GameFramework/RootMotionSource.cpp b/Engine/Source/Runtime/Engine/Private/GameFramework/RootMotionSource.cpp
--- a/Engine/Source/Runtime/Engine/Private/GameFramework/RootMotionSource.cpp
+++ b/Engine/Source/Runtime/Engine/Private/GameFramework/RootMotionSource.cpp
@@ -734,6 +734,11 @@ FRootMotionSource_JumpForce::FRootMotionSource_JumpForce()
 	, TimeMappingCurve(nullptr)
 	, SavedHalfwayLocation(FVector::ZeroVector)
 {
+	// Don't allow partial end ticks. Jump forces are meant to provide velocity that
+	// carries through to the end of the jump, and if we do partial ticks at the very end,
+	// it means the provided velocity can be significantly reduced on the very last tick,
+	// resulting in lost momentum. This is not desirable for jumps.
+	Settings.SetFlag(ERootMotionSourceSettingsFlags::DisablePartialEndTick);
 }
 
 bool FRootMotionSource_JumpForce::IsTimeOutEnabled() const
@@ -1180,7 +1185,7 @@ void FRootMotionSourceGroup::PrepareRootMotion(float DeltaTime, const ACharacter
 						}
 
 						// End of root motion
-						if (RootMotionSource->IsTimeOutEnabled())
+						if (RootMotionSource->IsTimeOutEnabled() && !RootMotionSource->Settings.HasFlag(ERootMotionSourceSettingsFlags::DisablePartialEndTick))
 						{
 							const float Duration = RootMotionSource->GetDuration();
 							if (RootMotionSource->GetTime() + SimulationTime >= Duration)
diff --git a/Engine/Source/Runtime/Engine/Private/GameUserSettings.cpp b/Engine/Source/Runtime/Engine/Private/GameUserSettings.cpp
--- a/Engine/Source/Runtime/Engine/Private/GameUserSettings.cpp
+++ b/Engine/Source/Runtime/Engine/Private/GameUserSettings.cpp
@@ -42,7 +42,16 @@ FIntPoint UGameUserSettings::GetDesktopResolution() const
 	}
 	else
 	{
-		FDisplayMetrics::GetDisplayMetrics(DisplayMetrics);
+		if (FApp::CanEverRender())
+		{
+			FDisplayMetrics::GetDisplayMetrics(DisplayMetrics);
+		}
+		else
+		{
+			// assume headless
+			DisplayMetrics.PrimaryDisplayWidth = 0;
+			DisplayMetrics.PrimaryDisplayHeight = 0;
+		}
 	}
 	return FIntPoint(DisplayMetrics.PrimaryDisplayWidth, DisplayMetrics.PrimaryDisplayHeight);
 }
diff --git a/Engine/Source/Runtime/Engine/Private/Materials/ParameterCollection.cpp b/Engine/Source/Runtime/Engine/Private/Materials/ParameterCollection.cpp
--- a/Engine/Source/Runtime/Engine/Private/Materials/ParameterCollection.cpp
+++ b/Engine/Source/Runtime/Engine/Private/Materials/ParameterCollection.cpp
@@ -31,6 +31,67 @@ void UMaterialParameterCollection::PostLoad()
 
 #if WITH_EDITOR
 
+template<typename ParameterType>
+FName CreateUniqueName(TArray<ParameterType>& Parameters, int32 RenameParameterIndex)
+{
+	FString RenameString;
+	Parameters[RenameParameterIndex].ParameterName.ToString(RenameString);
+
+	int32 NumberStartIndex = RenameString.FindLastCharByPredicate([](TCHAR Letter){ return !FChar::IsDigit(Letter); }) + 1;
+	
+	int32 RenameNumber = 0;
+	if (NumberStartIndex < RenameString.Len() - 1)
+	{
+		FString RenameStringNumberPart = RenameString.RightChop(NumberStartIndex);
+		ensure(RenameStringNumberPart.IsNumeric());
+
+		TTypeFromString<int32>::FromString(RenameNumber, *RenameStringNumberPart);
+	}
+
+	FString BaseString = RenameString.Left(NumberStartIndex);
+
+	FName Renamed = FName(*FString::Printf(TEXT("%s%u"), *BaseString, ++RenameNumber));
+
+	bool bMatchFound = false;
+	
+	do
+	{
+		bMatchFound = false;
+
+		for (int32 i = 0; i < Parameters.Num(); ++i)
+		{
+			if (Parameters[i].ParameterName == Renamed && RenameParameterIndex != i)
+			{
+				Renamed = FName(*FString::Printf(TEXT("%s%u"), *BaseString, ++RenameNumber));
+				bMatchFound = true;
+				break;
+			}
+		}
+	} while (bMatchFound);
+	
+	return Renamed;
+}
+
+template<typename ParameterType>
+void SanatizeParameters(TArray<ParameterType>& Parameters)
+{
+	for (int32 i = 0; i < Parameters.Num() - 1; ++i)
+	{
+		for (int32 j = i + 1; j < Parameters.Num(); ++j)
+		{
+			if (Parameters[i].Id == Parameters[j].Id)
+			{
+				FPlatformMisc::CreateGuid(Parameters[j].Id);
+			}
+
+			if (Parameters[i].ParameterName == Parameters[j].ParameterName)
+			{
+				Parameters[j].ParameterName = CreateUniqueName(Parameters, j);
+			}
+		}
+	}
+}
+
 TArray<FCollectionScalarParameter> PreviousScalarParameters;
 TArray<FCollectionVectorParameter> PreviousVectorParameters;
 
@@ -44,6 +105,9 @@ void UMaterialParameterCollection::PreEditChange(class FEditPropertyChain& Prope
 
 void UMaterialParameterCollection::PostEditChangeProperty(FPropertyChangedEvent& PropertyChangedEvent)
 {
+	SanatizeParameters(ScalarParameters);
+	SanatizeParameters(VectorParameters);
+
 	// If the array counts have changed, an element has been added or removed, and we need to update the uniform buffer layout,
 	// Which also requires recompiling any referencing materials
 	if (ScalarParameters.Num() != PreviousScalarParameters.Num()
@@ -162,61 +226,6 @@ void UMaterialParameterCollection::PostEditChangeProperty(FPropertyChangedEvent&
 	Super::PostEditChangeProperty(PropertyChangedEvent);
 }
 
-/**
- * Helper function for creating unique item names within a list of existing items
- *   InBaseName - Desired name prefix (will generate Prefix<N>)
- *   InExistingItems - Array of existing items which we want to ensure uniqueness within
- *   OutName - Target FName for result
- *   InNewIndex - Index of value that has just been added, so we can make sure we don't check against ourselves
- **/
-template <class T>
-inline void CreateUniqueName(const TCHAR* InBaseName, TArray<T>& InExistingItems, FName& OutName, int32 InNewIndex)
-{
-	int32 Index = 0;
-	bool bMatchFound = true;
-
-	while (bMatchFound)
-	{
-		bMatchFound = false;
-		OutName = FName(*FString::Printf(TEXT("%s%u"), InBaseName, Index++));
-
-		for (int32 i = 0; i < InExistingItems.Num(); ++i)
-		{
-			if (i != InNewIndex && InExistingItems[i].ParameterName == OutName)
-			{
-				bMatchFound = true;
-				break;
-			}
-		}
-	}
-}
-
-void UMaterialParameterCollection::PostEditChangeChainProperty(struct FPropertyChangedChainEvent& PropertyChangedEvent)
-{
-	// Auto-populate with useful parameter names
-	if (ScalarParameters.Num() > PreviousScalarParameters.Num())
-	{
-		int32 NewArrayIndex = PropertyChangedEvent.GetArrayIndex("ScalarParameters");
-
-		if (ScalarParameters.IsValidIndex(NewArrayIndex))
-		{
-			CreateUniqueName<FCollectionScalarParameter>(TEXT("Scalar"), ScalarParameters, ScalarParameters[NewArrayIndex].ParameterName, NewArrayIndex);
-		}
-	}
-
-	if (VectorParameters.Num() > PreviousVectorParameters.Num())
-	{
-		int32 NewArrayIndex = PropertyChangedEvent.GetArrayIndex("VectorParameters");
-
-		if (VectorParameters.IsValidIndex(NewArrayIndex))
-		{
-			CreateUniqueName<FCollectionVectorParameter>(TEXT("Vector"), VectorParameters, VectorParameters[NewArrayIndex].ParameterName, NewArrayIndex);
-		}
-	}
-
-	Super::PostEditChangeChainProperty(PropertyChangedEvent);
-}
-
 #endif // WITH_EDITOR
 
 FName UMaterialParameterCollection::GetParameterName(const FGuid& Id) const
diff --git a/Engine/Source/Runtime/Engine/Private/NetworkDriver.cpp b/Engine/Source/Runtime/Engine/Private/NetworkDriver.cpp
--- a/Engine/Source/Runtime/Engine/Private/NetworkDriver.cpp
+++ b/Engine/Source/Runtime/Engine/Private/NetworkDriver.cpp
@@ -147,6 +147,7 @@ UNetDriver::UNetDriver(const FObjectInitializer& ObjectInitializer)
 ,	OutOutOfOrderPackets(0)
 ,	StatUpdateTime(0.0)
 ,	StatPeriod(1.f)
+,	LastCleanupTime(0.0)
 ,	NetTag(0)
 ,	DebugRelevantActors(false)
 ,	ProcessQueuedBunchesCurrentFrameMilliseconds(0.0f)
@@ -301,9 +302,9 @@ void UNetDriver::TickFlush(float DeltaSeconds)
 	// Reset queued bunch amortization timer
 	ProcessQueuedBunchesCurrentFrameMilliseconds = 0.0f;
 
-#if USE_SERVER_PERF_COUNTERS || STATS
 	const double CurrentRealtimeSeconds = FPlatformTime::Seconds();
 
+#if USE_SERVER_PERF_COUNTERS || STATS
 	// Update network stats (only main game net driver for now) if stats or perf counters are used
 	if (NetDriverName == NAME_GameNetDriver && 
 		 CurrentRealtimeSeconds - StatUpdateTime > StatPeriod)
@@ -315,6 +316,7 @@ void UNetDriver::TickFlush(float DeltaSeconds)
 		int32 ClientOutBytesMin = 0;
 		int32 ClientOutBytesAvg = 0;
 		int NumClients = 0;
+		float RemoteSaturationMax = 0.0f;
 
 		// these need to be updated even if we are not collecting stats, since they get reported to analytics/QoS
 		for (UNetConnection * Client : ClientConnections)
@@ -356,7 +358,6 @@ void UNetDriver::TickFlush(float DeltaSeconds)
 		int32 UnAckCount = 0;
 		int32 PendingCount = 0;
 		int32 NetSaturated = 0;
-		float RemoteSaturationMax = 0.0f;
 
 		if (FThreadStats::IsCollectingData())
 		{
@@ -563,6 +564,7 @@ void UNetDriver::TickFlush(float DeltaSeconds)
 			PerfCounters->Set(TEXT("OutRateClientAvg"), ClientOutBytesAvg);
 
 			PerfCounters->Set(TEXT("ServerReplicateActorsTimeMs"), ServerReplicateActorsTimeMs);
+			PerfCounters->Set(TEXT("OutSaturationMax"), RemoteSaturationMax);
 		}
 #endif // USE_SERVER_PERF_COUNTERS
 
@@ -631,6 +633,24 @@ void UNetDriver::TickFlush(float DeltaSeconds)
 		}
 	}
 
+	// Go over RepChangedPropertyTrackerMap periodicallly, and remove entries that no longer have valid objects
+	// Unfortunately if you mark an object as pending kill, it will no longer find itself in this map,
+	// so we do this as a fail safe to make sure we never leak memory from this map
+	const double CleanupTimeSeconds = 10.0;
+
+	if ( CurrentRealtimeSeconds - LastCleanupTime > CleanupTimeSeconds )
+	{
+		for ( auto It = RepChangedPropertyTrackerMap.CreateIterator(); It; ++It )
+		{
+			if ( !It.Key().IsValid() )
+			{
+				It.RemoveCurrent();
+			}
+		}
+
+		LastCleanupTime = CurrentRealtimeSeconds;
+	}
+
 	// Update the lag state
 	UpdateNetworkLagState();
 }
diff --git a/Engine/Source/Runtime/Engine/Private/OnlineReplStructs.cpp b/Engine/Source/Runtime/Engine/Private/OnlineReplStructs.cpp
--- a/Engine/Source/Runtime/Engine/Private/OnlineReplStructs.cpp
+++ b/Engine/Source/Runtime/Engine/Private/OnlineReplStructs.cpp
@@ -25,19 +25,24 @@ FArchive& operator<<( FArchive& Ar, FUniqueNetIdRepl& UniqueNetId)
 			FString Contents;
 			Ar << Contents;	// that takes care about possible overflow
 
-			// Don't need to distinguish OSS interfaces here with world because we just want the create function below
-			IOnlineIdentityPtr IdentityInt = Online::GetIdentityInterface();
-			if (IdentityInt.IsValid())
-			{
-				TSharedPtr<const FUniqueNetId> UniqueNetIdPtr = IdentityInt->CreateUniquePlayerId(Contents);
-				UniqueNetId.SetUniqueNetId(UniqueNetIdPtr);
-			}
+			UniqueNetId.UniqueIdFromString(Contents);
 		}
 	}
 
 	return Ar;
 }
 
+void FUniqueNetIdRepl::UniqueIdFromString(const FString& Contents)
+{
+	// Don't need to distinguish OSS interfaces here with world because we just want the create function below
+	IOnlineIdentityPtr IdentityInt = Online::GetIdentityInterface();
+	if (IdentityInt.IsValid())
+	{
+		TSharedPtr<const FUniqueNetId> UniqueNetIdPtr = IdentityInt->CreateUniquePlayerId(Contents);
+		SetUniqueNetId(UniqueNetIdPtr);
+	}
+}
+
 bool FUniqueNetIdRepl::NetSerialize(FArchive& Ar, class UPackageMap* Map, bool& bOutSuccess)
 {
 	Ar << *this;
@@ -62,8 +67,59 @@ bool FUniqueNetIdRepl::ExportTextItem(FString& ValueStr, FUniqueNetIdRepl const&
 	return true;
 }
 
+bool FUniqueNetIdRepl::ImportTextItem(const TCHAR*& Buffer, int32 PortFlags, UObject* Parent, FOutputDevice* ErrorText)
+{
+	SetUniqueNetId(nullptr);
+
+	if (IsRunningCommandlet())
+	{
+#if !NO_LOGGING
+		ErrorText->Log(LogNet.GetCategoryName(), ELogVerbosity::Warning, TEXT("Invalid attempt to serialize FUniqueNetIdRepl during cooking"));
+#endif
+		return true;
+	}
+
+	if (Buffer)
+	{
+		FString Contents(Buffer);
+		UniqueIdFromString(Contents);
+	}
+
+	if (!IsValid())
+	{
+#if !NO_LOGGING
+		ErrorText->Log(LogNet.GetCategoryName(), ELogVerbosity::Warning, TEXT("Failed to import text to FUniqueNetIdRepl"));
+#endif
+	}
+
+	return true;
+}
+
+TSharedRef<FJsonValue> FUniqueNetIdRepl::ToJson() const
+{
+	if (IsValid())
+	{
+		return MakeShareable(new FJsonValueString(ToString()));
+	}
+	else
+	{
+		return MakeShareable(new FJsonValueString(TEXT("INVALID")));
+	}
+}
+
+void FUniqueNetIdRepl::FromJson(const FString& Json)
+{
+	SetUniqueNetId(nullptr);
+
+	if (!Json.IsEmpty())
+	{
+		UniqueIdFromString(Json);
+	}
+}
+
 void TestUniqueIdRepl(UWorld* InWorld)
 {
+#if !UE_BUILD_SHIPPING
 	bool bSuccess = true;
 
 	IOnlineIdentityPtr IdentityPtr = Online::GetIdentityInterface(InWorld);
@@ -79,7 +135,7 @@ void TestUniqueIdRepl(UWorld* InWorld)
 		}
 
 		FUniqueNetIdRepl ValidIdIn(UserId);
-		if (!ValidIdIn.IsValid() || UserId != ValidIdIn.GetUniqueNetId())
+		if (!ValidIdIn.IsValid() || UserId != ValidIdIn.GetUniqueNetId() || *UserId != *ValidIdIn)
 		{
 			UE_LOG(LogNet, Warning, TEXT("UserId input %s != UserId output %s"), *UserId->ToString(), *ValidIdIn->ToString());
 			bSuccess = false;
@@ -126,12 +182,26 @@ void TestUniqueIdRepl(UWorld* InWorld)
 				}
 			}
 		}
+
+		if (bSuccess)
+		{
+			FString OutString;
+			TSharedRef<FJsonValue> JsonValue = ValidIdIn.ToJson();
+			bSuccess = JsonValue->TryGetString(OutString);
+			if (bSuccess)
+			{
+				FUniqueNetIdRepl NewIdOut;
+				NewIdOut.FromJson(OutString);
+				bSuccess = NewIdOut.IsValid();
+			}
+		}
 	}
 
 	if (!bSuccess)
 	{
 		UE_LOG(LogNet, Warning, TEXT("TestUniqueIdRepl test failure!"));
 	}
+#endif
 }
 
 
diff --git a/Engine/Source/Runtime/Engine/Private/Particles/ParticleComponents.cpp b/Engine/Source/Runtime/Engine/Private/Particles/ParticleComponents.cpp
--- a/Engine/Source/Runtime/Engine/Private/Particles/ParticleComponents.cpp
+++ b/Engine/Source/Runtime/Engine/Private/Particles/ParticleComponents.cpp
@@ -3044,6 +3044,7 @@ UParticleSystemComponent::UParticleSystemComponent(const FObjectInitializer& Obj
 #if WITH_EDITORONLY_DATA
 	EditorDetailMode = -1;
 #endif // WITH_EDITORONLY_DATA
+	LastCheckedDetailMode = -1;
 	SetCollisionProfileName(UCollisionProfile::NoCollision_ProfileName);
 	bGenerateOverlapEvents = false;
 
@@ -4202,7 +4203,7 @@ void UParticleSystemComponent::TickComponent(float DeltaTime, enum ELevelTick Ti
 	}
 
 	// System settings may have been lowered. Support late deactivation.
-	int32 DetailModeCVar = GetCachedScalabilityCVars().DetailMode;
+	int32 DetailModeCVar = GetCurrentDetailMode();
 	const bool bDetailModeAllowsRendering	= DetailMode <= DetailModeCVar;
 	if ( bDetailModeAllowsRendering == false )
 	{
@@ -4233,8 +4234,16 @@ void UParticleSystemComponent::TickComponent(float DeltaTime, enum ELevelTick Ti
 
 		if (bRequiresReset)
 		{
+			bool bOldActive = bIsActive;
 			ResetParticles(true);
-			InitializeSystem();
+			if (bOldActive)
+			{
+				ActivateSystem();
+			}
+			else
+			{
+				InitializeSystem();
+			}
 		}
 
 		// Save the detail mode we just checked
@@ -4244,30 +4253,35 @@ void UParticleSystemComponent::TickComponent(float DeltaTime, enum ELevelTick Ti
 	// Bail out if MaxSecondsBeforeInactive > 0 and we haven't been rendered the last MaxSecondsBeforeInactive seconds.
 	if (bWarmingUp == false)
 	{
-		const float MaxSecondsBeforeInactive = FMath::Max( SecondsBeforeInactive, Template->SecondsBeforeInactive );
+		//For now, we're only allowing the SecondsBeforeInactive optimization on looping emitters as it can cause leaks with non-looping effects.
+		//Longer term, there is likely a better solution.
+		if (Template->IsLooping())
+		{
+			const float MaxSecondsBeforeInactive = FMath::Max(SecondsBeforeInactive, Template->SecondsBeforeInactive);
 
-		// Clamp MaxSecondsBeforeInactive to be at least twice the maximum
-		// smoothed frame time (45.45ms) because the rendering thread runs one 
-		// frame behind the game thread and so smaller time differences cannot 
-		// be reliably detected.
-		const float ClampedMaxSecondsBeforeInactive = MaxSecondsBeforeInactive > 0 ? FMath::Max(MaxSecondsBeforeInactive, 0.1f) : 0;
+			// Clamp MaxSecondsBeforeInactive to be at least twice the maximum
+			// smoothed frame time (45.45ms) because the rendering thread runs one 
+			// frame behind the game thread and so smaller time differences cannot 
+			// be reliably detected.
+			const float ClampedMaxSecondsBeforeInactive = MaxSecondsBeforeInactive > 0 ? FMath::Max(MaxSecondsBeforeInactive, 0.1f) : 0;
 
-		if( ClampedMaxSecondsBeforeInactive > 0 
-			&&	AccumTickTime > ClampedMaxSecondsBeforeInactive//SecondsBeforeInactive
-			&&	GetWorld()->IsGameWorld() )
-		{
-			SCOPE_CYCLE_COUNTER(STAT_UParticleSystemComponent_LOD_Inactive);
-			const float CurrentTimeSeconds = World->GetTimeSeconds();
-			if( CurrentTimeSeconds > (LastRenderTime + ClampedMaxSecondsBeforeInactive) )
+			if (ClampedMaxSecondsBeforeInactive > 0
+				&& AccumTickTime > ClampedMaxSecondsBeforeInactive//SecondsBeforeInactive
+				&&	GetWorld()->IsGameWorld())
 			{
-				bForcedInActive = true;
+				SCOPE_CYCLE_COUNTER(STAT_UParticleSystemComponent_LOD_Inactive);
+				const float CurrentTimeSeconds = World->GetTimeSeconds();
+				if (CurrentTimeSeconds > (LastRenderTime + ClampedMaxSecondsBeforeInactive))
+				{
+					bForcedInActive = true;
 
-				SpawnEvents.Empty();
-				DeathEvents.Empty();
-				CollisionEvents.Empty();
-				KismetEvents.Empty();
+					SpawnEvents.Empty();
+					DeathEvents.Empty();
+					CollisionEvents.Empty();
+					KismetEvents.Empty();
 
-				return;
+					return;
+				}
 			}
 		}
 
@@ -4671,7 +4685,7 @@ void UParticleSystemComponent::InitParticles(bool bReInitExistingEmitters)
 		WarmupTime = Template->WarmupTime;
 		WarmupTickRate = Template->WarmupTickRate;
 		bIsViewRelevanceDirty = true;
-		const int32 GlobalDetailMode = GetCachedScalabilityCVars().DetailMode;
+		const int32 GlobalDetailMode = GetCurrentDetailMode();;
 		const bool bCanEverRender = CanEverRender();
 
 		//simplified version.
@@ -4988,7 +5002,7 @@ void UParticleSystemComponent::ActivateSystem(bool bFlagAsJustAttached)
 	}
 
 	// System settings may have been lowered. Support late deactivation.
-	const bool bDetailModeAllowsRendering	= DetailMode <= GetCachedScalabilityCVars().DetailMode;
+	const bool bDetailModeAllowsRendering	= DetailMode <= GetCurrentDetailMode();
 
 	//Assume significant initially.
 	LastSignificantTime = GetWorld()->GetTimeSeconds();
@@ -5523,7 +5537,7 @@ void UParticleSystemComponent::InitializeSystem()
 		Template != NULL ? *Template->GetName() : TEXT("NULL"), this, FXSystem);
 
 	// System settings may have been lowered. Support late deactivation.
-	const bool bDetailModeAllowsRendering	= DetailMode <= GetCachedScalabilityCVars().DetailMode;
+	const bool bDetailModeAllowsRendering	= DetailMode <= GetCurrentDetailMode();
 
 	if( GIsAllowingParticles && bDetailModeAllowsRendering )
 	{
diff --git a/Engine/Source/Runtime/Engine/Private/PerfCountersHelpers.cpp b/Engine/Source/Runtime/Engine/Private/PerfCountersHelpers.cpp
--- a/Engine/Source/Runtime/Engine/Private/PerfCountersHelpers.cpp
+++ b/Engine/Source/Runtime/Engine/Private/PerfCountersHelpers.cpp
@@ -24,6 +24,15 @@ void ENGINE_API PerfCountersSet(const FString& Name, int32 Val, uint32 Flags)
 	}
 }
 
+void ENGINE_API PerfCountersSet(const FString& Name, const FString& Val, uint32 Flags)
+{
+	IPerfCounters* PerfCounters = IPerfCountersModule::Get().GetPerformanceCounters();
+	if (PerfCounters)
+	{
+		PerfCounters->Set(Name, Val, Flags);
+	}
+}
+
 int32 ENGINE_API PerfCountersIncrement(const FString & Name, int32 Add, int32 DefaultValue, uint32 Flags)
 {
 	IPerfCounters* PerfCounters = IPerfCountersModule::Get().GetPerformanceCounters();
diff --git a/Engine/Source/Runtime/Engine/Private/RepLayout.cpp b/Engine/Source/Runtime/Engine/Private/RepLayout.cpp
--- a/Engine/Source/Runtime/Engine/Private/RepLayout.cpp
+++ b/Engine/Source/Runtime/Engine/Private/RepLayout.cpp
@@ -782,7 +782,7 @@ void FRepLayout::UpdateChangelistHistory( FRepState * RepState, UClass * ObjectC
 	// If our buffer is currently full, forcibly send the entire history
 	if ( DumpHistory )
 	{
-		UE_LOG( LogRep, Warning, TEXT( "FRepLayout::UpdateChangelistHistory: History overflow, forcing history dump %s" ), *ObjectClass->GetName() );
+		UE_LOG( LogRep, Log, TEXT( "FRepLayout::UpdateChangelistHistory: History overflow, forcing history dump %s" ), *ObjectClass->GetName() );
 	}
 
 	for ( int32 i = RepState->HistoryStart; i < RepState->HistoryEnd; i++ )
diff --git a/Engine/Source/Runtime/Engine/Private/UserInterface/Canvas.cpp b/Engine/Source/Runtime/Engine/Private/UserInterface/Canvas.cpp
--- a/Engine/Source/Runtime/Engine/Private/UserInterface/Canvas.cpp
+++ b/Engine/Source/Runtime/Engine/Private/UserInterface/Canvas.cpp
@@ -2011,17 +2011,23 @@ void FDisplayDebugManager::DrawString(const FString& InDebugString, const float&
 {
 	if (Canvas)
 	{
+		const float TextScale = GetTextScale();
+		DebugTextItem.Scale = FVector2D(TextScale, TextScale);
+
 		DebugTextItem.Text = FText::FromString(InDebugString);
 		Canvas->DrawItem(DebugTextItem, FVector2D(CurrentPos.X + OptionalXOffset, CurrentPos.Y));
 
 		NextColumXPos = FMath::Max(NextColumXPos, CurrentPos.X + OptionalXOffset + DebugTextItem.DrawnSize.X);
-		MaxCharHeight = FMath::Max(MaxCharHeight, DebugTextItem.DrawnSize.Y);
-
-		CurrentPos.Y += GetYStep();
+		CurrentPos.Y += FMath::Max(GetYStep(), DebugTextItem.DrawnSize.Y);
 		AddColumnIfNeeded();
 	}
 }
 
+float FDisplayDebugManager::GetTextScale() const
+{
+	return Canvas ? FMath::Max(Canvas->SizeX / 1920.f, 1.f) : 1.f;
+}
+
 void FDisplayDebugManager::AddColumnIfNeeded()
 {
 	if (Canvas)
diff --git a/Engine/Source/Runtime/Engine/Public/Animation/AnimCurveTypes.h b/Engine/Source/Runtime/Engine/Public/Animation/AnimCurveTypes.h
--- a/Engine/Source/Runtime/Engine/Public/Animation/AnimCurveTypes.h
+++ b/Engine/Source/Runtime/Engine/Public/Animation/AnimCurveTypes.h
@@ -462,6 +462,7 @@ struct FBaseBlendedCurve
 	template <typename OtherAllocator>
 	void CopyFrom(const FBaseBlendedCurve<OtherAllocator>& CurveToCopyFrom)
 	{
+		checkf(CurveToCopyFrom.IsValid(), TEXT("Copying data from an invalid curve UIDList: 0x%x  (Sizes %i/%i)"), CurveToCopyFrom.UIDList, (CurveToCopyFrom.UIDList ? CurveToCopyFrom.UIDList->Num() : -1), CurveToCopyFrom.Elements.Num());
 		UIDList = CurveToCopyFrom.UIDList;
 		Elements.Reset();
 		Elements.Append(CurveToCopyFrom.Elements);
@@ -472,6 +473,7 @@ struct FBaseBlendedCurve
 	{
 		if (&CurveToCopyFrom != this)
 		{
+			checkf(CurveToCopyFrom.IsValid(), TEXT("Copying data from an invalid curve UIDList: 0x%x  (Sizes %i/%i)"), CurveToCopyFrom.UIDList, (CurveToCopyFrom.UIDList ? CurveToCopyFrom.UIDList->Num() : -1), CurveToCopyFrom.Elements.Num());
 			UIDList = CurveToCopyFrom.UIDList;
 			Elements.Reset();
 			Elements.Append(CurveToCopyFrom.Elements);
@@ -495,6 +497,13 @@ struct FBaseBlendedCurve
 		Elements.Reset();
 		Elements.Reserve(Count);
 	}
+
+	// Only checks bare minimal validity. (namely that we have a UID list and that it 
+	// is the same size as our element list
+	bool IsValid() const
+	{
+		return UIDList && (Elements.Num() == UIDList->Num());
+	}
 };
 
 struct FBlendedHeapCurve;
diff --git a/Engine/Source/Runtime/Engine/Public/ContentStreaming.h b/Engine/Source/Runtime/Engine/Public/ContentStreaming.h
--- a/Engine/Source/Runtime/Engine/Public/ContentStreaming.h
+++ b/Engine/Source/Runtime/Engine/Public/ContentStreaming.h
@@ -737,7 +737,7 @@ struct FSpawnedTextureInstance
  */
 struct FStreamMemoryTracker
 {
-#if WINVER < 0x0600
+#if PLATFORM_WINDOWS && (WINVER < 0x0600)
 	// Windows XP does not have InterlockedIncrement64
 	typedef int32 TSize;
 #else
diff --git a/Engine/Source/Runtime/Engine/Public/Net/PerfCountersHelpers.h b/Engine/Source/Runtime/Engine/Public/Net/PerfCountersHelpers.h
--- a/Engine/Source/Runtime/Engine/Public/Net/PerfCountersHelpers.h
+++ b/Engine/Source/Runtime/Engine/Public/Net/PerfCountersHelpers.h
@@ -14,6 +14,10 @@
 	*/
 	void ENGINE_API PerfCountersSet(const FString& Name, int32 Val, uint32 Flags = 0);
 
+	/** @brief Helper function for setting a performance counter. Compiled out if PerfCounterModule isn't used.
+	*/
+	void ENGINE_API PerfCountersSet(const FString& Name, const FString& Val, uint32 Flags = 0);
+
 	/** @brief Helper function for incrementing a performance counter. Compiled out if PerfCounterModule isn't used.
 	 *
 	 *  @param Name the name of the counter
@@ -47,6 +51,7 @@
 
 	FORCEINLINE void PerfCountersSet(const FString& Name, float Val, uint32 Flags = 0) { /* no-op */ };
 	FORCEINLINE void PerfCountersSet(const FString& Name, int32 Val, uint32 Flags = 0) { /* no-op */ };
+	FORCEINLINE void PerfCountersSet(const FString& Name, const FString& Val, uint32 Flags = 0) { /* no-op */ };
 	FORCEINLINE int32 PerfCountersIncrement(const FString & Name, int32 Add = 1, int32 DefaultValue = 0, uint32 Flags = 0) { return DefaultValue + Add; };
 	FORCEINLINE float ENGINE_API PerfCountersGet(const FString& Name, float DefaultVal) { return DefaultVal; }
 	FORCEINLINE double ENGINE_API PerfCountersGet(const FString& Name, double DefaultVal) { return DefaultVal; }
diff --git a/Engine/Source/Runtime/Foliage/Private/ProceduralFoliageComponent.cpp b/Engine/Source/Runtime/Foliage/Private/ProceduralFoliageComponent.cpp
--- a/Engine/Source/Runtime/Foliage/Private/ProceduralFoliageComponent.cpp
+++ b/Engine/Source/Runtime/Foliage/Private/ProceduralFoliageComponent.cpp
@@ -17,6 +17,10 @@ UProceduralFoliageComponent::UProceduralFoliageComponent(const FObjectInitialize
 	TileOverlap = 0.f;
 	ProceduralGuid = FGuid::NewGuid();
 #if WITH_EDITORONLY_DATA
+	bAllowLandscape = true;
+	bAllowBSP = true;
+	bAllowStaticMesh = true;
+	bAllowTranslucent = true;
 	bShowDebugTiles = false;
 #endif
 }
diff --git a/Engine/Source/Runtime/Foliage/Public/ProceduralFoliageComponent.h b/Engine/Source/Runtime/Foliage/Public/ProceduralFoliageComponent.h
--- a/Engine/Source/Runtime/Foliage/Public/ProceduralFoliageComponent.h
+++ b/Engine/Source/Runtime/Foliage/Public/ProceduralFoliageComponent.h
@@ -47,6 +47,23 @@ class FOLIAGE_API UProceduralFoliageComponent : public UActorComponent
 	float TileOverlap;
 
 #if WITH_EDITORONLY_DATA
+
+	/** Whether to place foliage on landscape */
+	UPROPERTY(Category = "ProceduralFoliage", BlueprintReadWrite, EditAnywhere)
+	bool bAllowLandscape;
+
+	/** Whether to place foliage on BSP */
+	UPROPERTY(Category = "ProceduralFoliage", BlueprintReadWrite, EditAnywhere)
+	bool bAllowBSP;
+
+	/** Whether to place foliage on StaticMesh */
+	UPROPERTY(Category = "ProceduralFoliage", BlueprintReadWrite, EditAnywhere)
+	bool bAllowStaticMesh;
+
+	/** Whether to place foliage on translucent geometry */
+	UPROPERTY(Category = "ProceduralFoliage", BlueprintReadWrite, EditAnywhere)
+	bool bAllowTranslucent;
+
 	/** Whether to visualize the tiles used for the foliage spawner simulation */
 	UPROPERTY(Category = "ProceduralFoliage", BlueprintReadWrite, EditAnywhere)
 	bool bShowDebugTiles;
diff --git a/Engine/Source/Runtime/GameplayAbilities/Private/Abilities/Tasks/AbilityTask_ApplyRootMotionMoveToForce.cpp b/Engine/Source/Runtime/GameplayAbilities/Private/Abilities/Tasks/AbilityTask_ApplyRootMotionMoveToForce.cpp
--- a/Engine/Source/Runtime/GameplayAbilities/Private/Abilities/Tasks/AbilityTask_ApplyRootMotionMoveToForce.cpp
+++ b/Engine/Source/Runtime/GameplayAbilities/Private/Abilities/Tasks/AbilityTask_ApplyRootMotionMoveToForce.cpp
@@ -18,9 +18,11 @@ UAbilityTask_ApplyRootMotionMoveToForce::UAbilityTask_ApplyRootMotionMoveToForce
 	MovementComponent = nullptr;
 	bRestrictSpeedToExpected = false;
 	PathOffsetCurve = nullptr;
+	VelocityOnFinishMode = EOrionRootMotionFinishVelocityMode::MaintainLastRootMotionVelocity;
+	SetVelocityOnFinish = FVector::ZeroVector;
 }
 
-UAbilityTask_ApplyRootMotionMoveToForce* UAbilityTask_ApplyRootMotionMoveToForce::ApplyRootMotionMoveToForce(UObject* WorldContextObject, FName TaskInstanceName, FVector TargetLocation, float Duration, bool bSetNewMovementMode, EMovementMode MovementMode, bool bRestrictSpeedToExpected, UCurveVector* PathOffsetCurve)
+UAbilityTask_ApplyRootMotionMoveToForce* UAbilityTask_ApplyRootMotionMoveToForce::ApplyRootMotionMoveToForce(UObject* WorldContextObject, FName TaskInstanceName, FVector TargetLocation, float Duration, bool bSetNewMovementMode, EMovementMode MovementMode, bool bRestrictSpeedToExpected, UCurveVector* PathOffsetCurve, EOrionRootMotionFinishVelocityMode VelocityOnFinishMode, FVector SetVelocityOnFinish)
 {
 	auto MyTask = NewAbilityTask<UAbilityTask_ApplyRootMotionMoveToForce>(WorldContextObject, TaskInstanceName);
 
@@ -31,6 +33,8 @@ UAbilityTask_ApplyRootMotionMoveToForce* UAbilityTask_ApplyRootMotionMoveToForce
 	MyTask->NewMovementMode = MovementMode;
 	MyTask->bRestrictSpeedToExpected = bRestrictSpeedToExpected;
 	MyTask->PathOffsetCurve = PathOffsetCurve;
+	MyTask->VelocityOnFinishMode = VelocityOnFinishMode;
+	MyTask->SetVelocityOnFinish = SetVelocityOnFinish;
 	if (MyTask->GetAvatarActor() != nullptr)
 	{
 		MyTask->StartLocation = MyTask->GetAvatarActor()->GetActorLocation();
@@ -153,6 +157,8 @@ void UAbilityTask_ApplyRootMotionMoveToForce::GetLifetimeReplicatedProps(TArray<
 	DOREPLIFETIME(UAbilityTask_ApplyRootMotionMoveToForce, NewMovementMode);
 	DOREPLIFETIME(UAbilityTask_ApplyRootMotionMoveToForce, bRestrictSpeedToExpected);
 	DOREPLIFETIME(UAbilityTask_ApplyRootMotionMoveToForce, PathOffsetCurve);
+	DOREPLIFETIME(UAbilityTask_ApplyRootMotionMoveToForce, VelocityOnFinishMode);
+	DOREPLIFETIME(UAbilityTask_ApplyRootMotionMoveToForce, SetVelocityOnFinish);
 }
 
 void UAbilityTask_ApplyRootMotionMoveToForce::PreDestroyFromReplication()
@@ -171,6 +177,19 @@ void UAbilityTask_ApplyRootMotionMoveToForce::OnDestroy(bool AbilityIsEnding)
 		{
 			MovementComponent->SetMovementMode(EMovementMode::MOVE_Falling);
 		}
+
+		if (VelocityOnFinishMode == EOrionRootMotionFinishVelocityMode::SetVelocity)
+		{
+			ACharacter* Character = MovementComponent->GetCharacterOwner();
+			if (Character)
+			{
+				MovementComponent->Velocity = Character->GetActorRotation().RotateVector(SetVelocityOnFinish);
+			}
+			else
+			{
+				MovementComponent->Velocity = SetVelocityOnFinish;
+			}
+		}
 	}
 
 	Super::OnDestroy(AbilityIsEnding);
diff --git a/Engine/Source/Runtime/GameplayAbilities/Private/GameplayCueManager.cpp b/Engine/Source/Runtime/GameplayAbilities/Private/GameplayCueManager.cpp
--- a/Engine/Source/Runtime/GameplayAbilities/Private/GameplayCueManager.cpp
+++ b/Engine/Source/Runtime/GameplayAbilities/Private/GameplayCueManager.cpp
@@ -1028,7 +1028,7 @@ void UGameplayCueManager::CheckForPreallocation(UClass* GCClass)
 {
 	if (AGameplayCueNotify_Actor* InstancedCue = Cast<AGameplayCueNotify_Actor>(GCClass->ClassDefaultObject))
 	{
-		if (InstancedCue->NumPreallocatedInstances > 0)
+		if (InstancedCue->NumPreallocatedInstances > 0 && GameplayCueClassesForPreallocation.Contains(InstancedCue) == false)
 		{
 			// Add this to the global list
 			GameplayCueClassesForPreallocation.Add(InstancedCue);
diff --git a/Engine/Source/Runtime/GameplayAbilities/Private/GameplayCueNotify_Actor.cpp b/Engine/Source/Runtime/GameplayAbilities/Private/GameplayCueNotify_Actor.cpp
--- a/Engine/Source/Runtime/GameplayAbilities/Private/GameplayCueNotify_Actor.cpp
+++ b/Engine/Source/Runtime/GameplayAbilities/Private/GameplayCueNotify_Actor.cpp
@@ -16,8 +16,13 @@ AGameplayCueNotify_Actor::AGameplayCueNotify_Actor(const FObjectInitializer& Obj
 	AutoDestroyDelay = 0.f;
 	bUniqueInstancePerSourceObject = false;
 	bUniqueInstancePerInstigator = false;
+	bAllowMultipleOnActiveEvents = true;
+	bAllowMultipleWhileActiveEvents = true;
 
 	NumPreallocatedInstances = 0;
+
+	bHasHandledOnActiveEvent = false;
+	bHasHandledWhileActiveEvent = false;
 }
 
 #if WITH_EDITOR
@@ -89,6 +94,19 @@ void AGameplayCueNotify_Actor::HandleGameplayCue(AActor* MyTarget, EGameplayCueE
 		ABILITY_LOG(Warning, TEXT("GameplayCue parameter is none for %s"), *GetNameSafe(this));
 	}
 
+	// Handle multiple event gating
+	{
+		if (EventType == EGameplayCueEvent::OnActive && !bAllowMultipleOnActiveEvents && bHasHandledOnActiveEvent)
+		{
+			return;
+		}
+
+		if (EventType == EGameplayCueEvent::WhileActive && !bAllowMultipleWhileActiveEvents && bHasHandledWhileActiveEvent)
+		{
+			return;
+		}
+	}
+
 	// If cvar is enabled, check that the target no longer has the matched tag before doing remove logic. This is a simple way of supporting stacking, such that if an actor has two sources giving him the same GC tag, it will not be removed when the first one is removed.
 	if (GameplayCueNotifyTagCheckOnRemove > 0 && EventType == EGameplayCueEvent::Removed)
 	{
@@ -112,10 +130,12 @@ void AGameplayCueNotify_Actor::HandleGameplayCue(AActor* MyTarget, EGameplayCueE
 		{
 		case EGameplayCueEvent::OnActive:
 			OnActive(MyTarget, Parameters);
+			bHasHandledOnActiveEvent = true;
 			break;
 
 		case EGameplayCueEvent::WhileActive:
 			WhileActive(MyTarget, Parameters);
+			bHasHandledWhileActiveEvent = true;
 			break;
 
 		case EGameplayCueEvent::Executed:
@@ -193,5 +213,7 @@ bool AGameplayCueNotify_Actor::GameplayCuePendingRemove()
 
 bool AGameplayCueNotify_Actor::Recycle()
 {
+	bHasHandledOnActiveEvent = false;
+	bHasHandledWhileActiveEvent = false;
 	return false;
 }
\ No newline at end of file
diff --git a/Engine/Source/Runtime/GameplayAbilities/Private/GameplayEffect.cpp b/Engine/Source/Runtime/GameplayAbilities/Private/GameplayEffect.cpp
--- a/Engine/Source/Runtime/GameplayAbilities/Private/GameplayEffect.cpp
+++ b/Engine/Source/Runtime/GameplayAbilities/Private/GameplayEffect.cpp
@@ -1483,6 +1483,11 @@ void FActiveGameplayEffect::PostReplicatedChange(const struct FActiveGameplayEff
 	const_cast<FActiveGameplayEffectsContainer&>(InArray).UpdateAllAggregatorModMagnitudes(*this);
 }
 
+void FActiveGameplayEffect::RecomputeStartWorldTime(const FActiveGameplayEffectsContainer& InArray)
+{
+	StartWorldTime = InArray.GetWorldTime() - static_cast<float>(InArray.GetServerWorldTime() - StartServerWorldTime);
+}
+
 // --------------------------------------------------------------------------------------------------------------------------------------------------------
 //
 //	FActiveGameplayEffectsContainer
diff --git a/Engine/Source/Runtime/GameplayAbilities/Public/Abilities/Tasks/AbilityTask_ApplyRootMotionMoveToForce.h b/Engine/Source/Runtime/GameplayAbilities/Public/Abilities/Tasks/AbilityTask_ApplyRootMotionMoveToForce.h
--- a/Engine/Source/Runtime/GameplayAbilities/Public/Abilities/Tasks/AbilityTask_ApplyRootMotionMoveToForce.h
+++ b/Engine/Source/Runtime/GameplayAbilities/Public/Abilities/Tasks/AbilityTask_ApplyRootMotionMoveToForce.h
@@ -7,6 +7,15 @@ DECLARE_DYNAMIC_MULTICAST_DELEGATE(FApplyRootMotionMoveToForceDelegate);
 
 class AActor;
 
+UENUM()
+enum class EOrionRootMotionFinishVelocityMode : uint8
+{
+	// Maintain the last velocity root motion gave to the character
+	MaintainLastRootMotionVelocity = 0,
+	// Set Velocity to the specified value (for example, 0,0,0 to stop the character)
+	SetVelocity,
+};
+
 /**
  *	Applies force to character's movement
  */
@@ -25,7 +34,7 @@ class UAbilityTask_ApplyRootMotionMoveToForce : public UAbilityTask
 
 	/** Apply force to character's movement */
 	UFUNCTION(BlueprintCallable, Category = "Ability|Tasks", meta = (HidePin = "WorldContextObject", DefaultToSelf = "WorldContextObject", BlueprintInternalUseOnly = "TRUE"))
-	static UAbilityTask_ApplyRootMotionMoveToForce* ApplyRootMotionMoveToForce(UObject* WorldContextObject, FName TaskInstanceName, FVector TargetLocation, float Duration, bool bSetNewMovementMode, EMovementMode MovementMode, bool bRestrictSpeedToExpected, UCurveVector* PathOffsetCurve);
+	static UAbilityTask_ApplyRootMotionMoveToForce* ApplyRootMotionMoveToForce(UObject* WorldContextObject, FName TaskInstanceName, FVector TargetLocation, float Duration, bool bSetNewMovementMode, EMovementMode MovementMode, bool bRestrictSpeedToExpected, UCurveVector* PathOffsetCurve, EOrionRootMotionFinishVelocityMode VelocityOnFinishMode, FVector SetVelocityOnFinish);
 
 	virtual void Activate() override;
 
@@ -57,12 +66,23 @@ protected:
 	UPROPERTY(Replicated)
 	TEnumAsByte<EMovementMode> NewMovementMode;
 
+	/** If enabled, we limit velocity to the initial expected velocity to go distance to the target over Duration.
+	 *  This prevents cases of getting really high velocity the last few frames of the root motion if you were being blocked by
+	 *  collision. Disabled means we do everything we can to velocity during the move to get to the TargetLocation. */
 	UPROPERTY(Replicated)
 	bool bRestrictSpeedToExpected;
 
 	UPROPERTY(Replicated)
 	UCurveVector* PathOffsetCurve;
 
+	/** What to do with character's Velocity when root motion finishes */
+	UPROPERTY(Replicated)
+	EOrionRootMotionFinishVelocityMode VelocityOnFinishMode;
+
+	/** If VelocityOnFinish mode is "SetVelocity", character velocity is set to this value when root motion finishes */
+	UPROPERTY(Replicated)
+	FVector SetVelocityOnFinish;
+
 	uint16 RootMotionSourceID;
 	EMovementMode PreviousMovementMode;
 
diff --git a/Engine/Source/Runtime/GameplayAbilities/Public/GameplayCueNotify_Actor.h b/Engine/Source/Runtime/GameplayAbilities/Public/GameplayCueNotify_Actor.h
--- a/Engine/Source/Runtime/GameplayAbilities/Public/GameplayCueNotify_Actor.h
+++ b/Engine/Source/Runtime/GameplayAbilities/Public/GameplayCueNotify_Actor.h
@@ -99,6 +99,20 @@ class GAMEPLAYABILITIES_API AGameplayCueNotify_Actor : public AActor
 	UPROPERTY(EditDefaultsOnly, Category = GameplayCue)
 	bool bUniqueInstancePerSourceObject;
 
+	/**
+	 *	Does this cue trigger its OnActive event if it's already been triggered?
+	 *  This can occur when the associated tag is triggered by multiple sources and there is no unique instancing.
+	 */ 	 
+	UPROPERTY(EditDefaultsOnly, Category = GameplayCue)
+	bool bAllowMultipleOnActiveEvents;
+
+	/**
+	 *	Does this cue trigger its WhileActive event if it's already been triggered?
+	 *  This can occur when the associated tag is triggered by multiple sources and there is no unique instancing.
+	 */ 	 
+	UPROPERTY(EditDefaultsOnly, Category = GameplayCue)
+	bool bAllowMultipleWhileActiveEvents;
+
 	/** How many instances of the gameplay cue to preallocate */
 	UPROPERTY(EditDefaultsOnly, Category = GameplayCue)
 	int32 NumPreallocatedInstances;
@@ -108,6 +122,9 @@ class GAMEPLAYABILITIES_API AGameplayCueNotify_Actor : public AActor
 protected:
 	FTimerHandle FinishTimerHandle;
 
+	bool bHasHandledOnActiveEvent;
+	bool bHasHandledWhileActiveEvent;
+
 private:
 	virtual void DeriveGameplayCueTagFromAssetName();
 
diff --git a/Engine/Source/Runtime/GameplayAbilities/Public/GameplayEffect.h b/Engine/Source/Runtime/GameplayAbilities/Public/GameplayEffect.h
--- a/Engine/Source/Runtime/GameplayAbilities/Public/GameplayEffect.h
+++ b/Engine/Source/Runtime/GameplayAbilities/Public/GameplayEffect.h
@@ -1288,6 +1288,9 @@ struct GAMEPLAYABILITIES_API FActiveGameplayEffect : public FFastArraySerializer
 	void PostReplicatedAdd(const struct FActiveGameplayEffectsContainer &InArray);
 	void PostReplicatedChange(const struct FActiveGameplayEffectsContainer &InArray);
 
+	/** Refreshes the cached StartWorldTime for this effect. To be used when the server/client world time delta changes significantly to keep the start time in sync. */
+	void RecomputeStartWorldTime(const FActiveGameplayEffectsContainer& InArray);
+
 	bool operator==(const FActiveGameplayEffect& Other)
 	{
 		return Handle == Other.Handle;
diff --git a/Engine/Source/Runtime/JsonUtilities/Private/JsonObjectConverter.cpp b/Engine/Source/Runtime/JsonUtilities/Private/JsonObjectConverter.cpp
--- a/Engine/Source/Runtime/JsonUtilities/Private/JsonObjectConverter.cpp
+++ b/Engine/Source/Runtime/JsonUtilities/Private/JsonObjectConverter.cpp
@@ -70,6 +70,14 @@ TSharedPtr<FJsonValue> ConvertScalarUPropertyToJsonValue(UProperty* Property, co
 	}
 	else if (UStructProperty *StructProperty = Cast<UStructProperty>(Property))
 	{
+		UScriptStruct::ICppStructOps* TheCppStructOps = StructProperty->Struct->GetCppStructOps();
+		if (TheCppStructOps && TheCppStructOps->HasExportTextItem())
+		{
+			FString OutValueStr;
+			TheCppStructOps->ExportTextItem(OutValueStr, Value, nullptr, nullptr, PPF_None, nullptr);
+			return MakeShareable(new FJsonValueString(OutValueStr));
+		}
+
 		TSharedRef<FJsonObject> Out = MakeShareable(new FJsonObject());
 		if (FJsonObjectConverter::UStructToJsonObject(StructProperty->Struct, Value, Out, CheckFlags & (~CPF_ParmFlags), SkipFlags, ExportCb))
 		{
@@ -405,6 +413,14 @@ bool ConvertScalarJsonValueToUProperty(TSharedPtr<FJsonValue> JsonValue, UProper
 				return false;
 			}
 		}
+		else if (JsonValue->Type == EJson::String && StructProperty->Struct->GetCppStructOps() && StructProperty->Struct->GetCppStructOps()->HasImportTextItem())
+		{
+			UScriptStruct::ICppStructOps* TheCppStructOps = StructProperty->Struct->GetCppStructOps();
+			
+			FString ImportTextString = JsonValue->AsString();
+			const TCHAR* ImportTextPtr = *ImportTextString;
+			TheCppStructOps->ImportTextItem(ImportTextPtr, OutValue, PPF_None, nullptr, (FOutputDevice*)GWarn);
+		}
 		else
 		{
 			UE_LOG(LogJson, Error, TEXT("JsonValueToUProperty - Attempted to import UStruct from non-object JSON key for property %s"), *Property->GetNameCPP());
diff --git a/Engine/Source/Runtime/Linux/LinuxCommonStartup/Private/LinuxCommonStartup.cpp b/Engine/Source/Runtime/Linux/LinuxCommonStartup/Private/LinuxCommonStartup.cpp
--- a/Engine/Source/Runtime/Linux/LinuxCommonStartup/Private/LinuxCommonStartup.cpp
+++ b/Engine/Source/Runtime/Linux/LinuxCommonStartup/Private/LinuxCommonStartup.cpp
@@ -13,10 +13,6 @@ static FString GSavedCommandLine;
 extern int32 GuardedMain( const TCHAR* CmdLine );
 extern void LaunchStaticShutdownAfterError();
 
-// FIXME: handle expose it someplace else?
-extern int32 DLLIMPORT ReportCrash(const FLinuxCrashContext& Context);
-extern void DLLIMPORT GenerateCrashInfoAndLaunchReporter(const FLinuxCrashContext& Context);
-
 /**
  * Game-specific crash reporter
  */
@@ -27,7 +23,8 @@ void CommonLinuxCrashHandler(const FGenericCrashContext& GenericContext)
 	const FLinuxCrashContext& Context = static_cast< const FLinuxCrashContext& >( GenericContext );
 	printf("CommonLinuxCrashHandler: Signal=%d\n", Context.Signal);
 
-	ReportCrash(Context);
+	// better than having mutable fields?
+	const_cast< FLinuxCrashContext& >(Context).CaptureStackTrace();
 	if (GLog)
 	{
 		GLog->Flush();
@@ -42,7 +39,7 @@ void CommonLinuxCrashHandler(const FGenericCrashContext& GenericContext)
 		GError->HandleError();
 	}
 
-	return GenerateCrashInfoAndLaunchReporter(Context);
+	return Context.GenerateCrashInfoAndLaunchReporter();
 }
 
 
diff --git a/Engine/Source/Runtime/NetworkReplayStreaming/HttpNetworkReplayStreaming/Private/HttpNetworkReplayStreaming.cpp b/Engine/Source/Runtime/NetworkReplayStreaming/HttpNetworkReplayStreaming/Private/HttpNetworkReplayStreaming.cpp
--- a/Engine/Source/Runtime/NetworkReplayStreaming/HttpNetworkReplayStreaming/Private/HttpNetworkReplayStreaming.cpp
+++ b/Engine/Source/Runtime/NetworkReplayStreaming/HttpNetworkReplayStreaming/Private/HttpNetworkReplayStreaming.cpp
@@ -4,6 +4,13 @@
 
 DEFINE_LOG_CATEGORY_STATIC( LogHttpReplay, Log, All );
 
+DECLARE_STATS_GROUP( TEXT( "HttpReplay" ), STATGROUP_HttpReplay, STATCAT_Advanced )
+
+DECLARE_DWORD_ACCUMULATOR_STAT( TEXT( "Http replay raw bytes" ), STAT_HttpReplay_Raw, STATGROUP_HttpReplay );
+DECLARE_DWORD_ACCUMULATOR_STAT( TEXT( "Http replay compressed bytes" ), STAT_HttpReplay_Compressed, STATGROUP_HttpReplay );
+DECLARE_CYCLE_STAT( TEXT( "Http replay compress time" ), STAT_HttpReplay_CompressTime, STATGROUP_HttpReplay );
+DECLARE_CYCLE_STAT( TEXT( "Http replay decompress time" ), STAT_HttpReplay_DecompressTime, STATGROUP_HttpReplay );
+
 static TAutoConsoleVariable<FString> CVarMetaFilterOverride( TEXT( "httpReplay.MetaFilterOverride" ), TEXT( "" ), TEXT( "" ) );
 
 class FNetworkReplayListItem : public FOnlineJsonSerializable
@@ -224,6 +231,8 @@ void FHttpNetworkReplayStreamer::StartStreaming( const FString& CustomName, cons
 	StreamTimeRangeStart	= 0;
 	StreamTimeRangeEnd		= 0;
 
+	EventGroupSet.Empty();
+
 	// Create the Http request and add to pending request list
 	TSharedRef<class IHttpRequest> HttpRequest = FHttpModule::Get().CreateRequest();
 
@@ -286,13 +295,14 @@ void FHttpNetworkReplayStreamer::StartStreaming( const FString& CustomName, cons
 
 		HttpRequest->OnProcessRequestComplete().BindRaw( this, &FHttpNetworkReplayStreamer::HttpStartUploadingFinished );
 
+		HttpRequest->SetHeader( TEXT( "Content-Type" ), TEXT( "application/json" ) );
+
 		if ( UserNames.Num() > 0 )
 		{
 			FNetworkReplayUserList UserList;
 
 			UserList.Users = UserNames;
 			HttpRequest->SetContentAsString( UserList.ToJson() );
-			HttpRequest->SetHeader( TEXT( "Content-Type" ), TEXT( "application/json" ) );
 		}
 
 		AddRequestToQueue( EQueuedHttpRequestType::StartUploading, HttpRequest );
@@ -428,7 +438,30 @@ void FHttpNetworkReplayStreamer::FlushStream()
 	HttpRequest->SetURL(FString::Printf(TEXT("%sreplay/%s/file/stream.%i?numChunks=%i&time=%i&mTime1=%i&mTime2=%i"), *ServerURL, *SessionName, StreamChunkIndex, StreamChunkIndex + 1, TotalDemoTimeInMS, StreamTimeRangeStart, StreamTimeRangeEnd));
 	HttpRequest->SetVerb( TEXT( "POST" ) );
 	HttpRequest->SetHeader( TEXT( "Content-Type" ), TEXT( "application/octet-stream" ) );
-	HttpRequest->SetContent( StreamArchive.Buffer );
+
+	if ( SupportsCompression() )
+	{
+		SCOPE_CYCLE_COUNTER( STAT_HttpReplay_CompressTime );
+
+		const double StartTime = FPlatformTime::Seconds();
+		FHttpStreamFArchive Compressed;
+
+		if ( !CompressBuffer( StreamArchive.Buffer, Compressed ) )
+		{
+			SetLastError( ENetworkReplayError::ServiceUnavailable );
+			return;
+		}
+
+		const double EndTime = FPlatformTime::Seconds();
+
+		UE_LOG( LogHttpReplay, VeryVerbose, TEXT( "Compressed stream. Original: %i, Compressed: %i, Time: %2.2f MS" ), StreamArchive.Buffer.Num(), Compressed.Buffer.Num(), ( EndTime - StartTime ) * 1000.0f );
+
+		HttpRequest->SetContent( Compressed.Buffer );
+	}
+	else
+	{
+		HttpRequest->SetContent( StreamArchive.Buffer );
+	}
 
 	StreamArchive.Buffer.Empty();
 	StreamArchive.Pos = 0;
@@ -672,7 +705,27 @@ void FHttpNetworkReplayStreamer::FlushCheckpointInternal( uint32 TimeInMS )
 	HttpRequest->SetURL( FString::Printf( TEXT( "%sreplay/%s/event?group=checkpoint&time1=%i&time2=%i&meta=%i" ), *ServerURL, *SessionName, TimeInMS, TimeInMS, StreamChunkIndex ) );
 	HttpRequest->SetVerb( TEXT( "POST" ) );
 	HttpRequest->SetHeader( TEXT( "Content-Type" ), TEXT( "application/octet-stream" ) );
-	HttpRequest->SetContent( CheckpointArchive.Buffer );
+
+	if ( SupportsCompression() )
+	{
+		SCOPE_CYCLE_COUNTER( STAT_HttpReplay_CompressTime );
+
+		const double StartTime = FPlatformTime::Seconds();
+		FHttpStreamFArchive Compressed;
+		if ( !CompressBuffer( CheckpointArchive.Buffer, Compressed ) )
+		{
+			SetLastError( ENetworkReplayError::ServiceUnavailable );
+			return;
+		}
+		const double EndTime = FPlatformTime::Seconds();
+		HttpRequest->SetContent( Compressed.Buffer );
+
+		UE_LOG( LogHttpReplay, VeryVerbose, TEXT( "Compressed checkpoint. Original: %i, Compressed: %i, Time: %2.2f MS" ), CheckpointArchive.Buffer.Num(), Compressed.Buffer.Num(), ( EndTime - StartTime ) * 1000.0f );
+	}
+	else
+	{
+		HttpRequest->SetContent( CheckpointArchive.Buffer );
+	}
 
 	CheckpointArchive.Buffer.Empty();
 	CheckpointArchive.Pos = 0;
@@ -746,6 +799,13 @@ void FHttpNetworkReplayStreamer::AddEvent( const uint32 TimeInMS, const FString&
 		return;
 	}
 
+	if ( StreamerState == EStreamerState::StreamingUp && !EventGroupSet.Contains( Group ) )
+	{
+		// Add the group as a user, so we can quickly find replays that have these event types in them
+		EventGroupSet.Add( Group );
+		AddUserToReplay( Group );
+	}
+
 	AddOrUpdateEvent( TEXT( "" ), TimeInMS, Group, Meta, Data );
 }
 
@@ -1117,6 +1177,35 @@ void FHttpNetworkReplayStreamer::EnumerateRecentStreams( const FNetworkReplayVer
 	AddRequestToQueue( EQueuedHttpRequestType::EnumeratingSessions, HttpRequest );
 }
 
+FQueuedHttpRequestAddUser::FQueuedHttpRequestAddUser( const FString& InUserName, TSharedRef< class IHttpRequest > InHttpRequest ) : FQueuedHttpRequest( EQueuedHttpRequestType::AddingUser, InHttpRequest )
+{
+	FNetworkReplayUserList UserList;
+
+	UserList.Users.Add( InUserName );
+	FString JsonString = UserList.ToJson();
+
+	Request->SetVerb( TEXT( "POST" ) );
+	Request->SetContentAsString( JsonString );
+	Request->SetHeader( TEXT( "Content-Type" ), TEXT( "application/json" ) );
+}
+
+bool FQueuedHttpRequestAddUser::PreProcess( FHttpNetworkReplayStreamer* Streamer, const FString& ServerURL, const FString& SessionName )
+{
+	if ( SessionName.IsEmpty() )
+	{
+		UE_LOG( LogHttpReplay, Warning, TEXT( "FQueuedHttpRequestAddUser::PreProcess. SessionName is empty." ) );
+		return false;
+	}
+
+	//
+	// Now that we have the session name, we can set the URL
+	//
+
+	Request->SetURL( FString::Printf( TEXT( "%sreplay/%s/users" ), *ServerURL, *SessionName ) );
+
+	return true;
+}
+
 void FHttpNetworkReplayStreamer::AddUserToReplay(const FString& UserString)
 {
 	if ( StreamerState != EStreamerState::StreamingUp )
@@ -1130,21 +1219,13 @@ void FHttpNetworkReplayStreamer::AddUserToReplay(const FString& UserString)
 		return;
 	}
 
-	TSharedRef<class IHttpRequest> HttpRequest = FHttpModule::Get().CreateRequest();
-
-	HttpRequest->SetURL( FString::Printf( TEXT( "%sreplay/%s/users" ), *ServerURL, *SessionName ) );
-	HttpRequest->SetVerb( TEXT( "POST" ) );
-
-	FNetworkReplayUserList UserList;
-
-	UserList.Users.Add(UserString);
-	FString JsonString = UserList.ToJson();
-	HttpRequest->SetContentAsString( JsonString );
-	HttpRequest->SetHeader( TEXT( "Content-Type" ), TEXT( "application/json" ) );
+	// Create the Http request and add to pending request list
+	TSharedRef< class IHttpRequest > HttpRequest = FHttpModule::Get().CreateRequest();
 
 	HttpRequest->OnProcessRequestComplete().BindRaw( this, &FHttpNetworkReplayStreamer::HttpAddUserFinished );
 
-	AddRequestToQueue( EQueuedHttpRequestType::AddingUser, HttpRequest );
+	// Add it as a custom event so we can snag the session name at the time of send (which we should have by then)
+	AddCustomRequestToQueue( TSharedPtr< FQueuedHttpRequest >( new FQueuedHttpRequestAddUser( UserString, HttpRequest ) ) );
 }
 
 void FHttpNetworkReplayStreamer::EnumerateCheckpoints()
@@ -1498,7 +1579,29 @@ void FHttpNetworkReplayStreamer::HttpDownloadFinished( FHttpRequestPtr HttpReque
 		{
 			if ( HttpResponse->GetContent().Num() > 0 )
 			{
-				StreamArchive.Buffer.Append( HttpResponse->GetContent() );
+				if ( SupportsCompression() )
+				{
+					SCOPE_CYCLE_COUNTER( STAT_HttpReplay_DecompressTime );
+
+					FHttpStreamFArchive Compressed;
+					Compressed.Buffer = HttpResponse->GetContent();
+					TArray< uint8 > Uncompressed;
+
+					if ( !DecompressBuffer( Compressed, Uncompressed ) )
+					{
+						StreamArchive.Buffer.Empty();
+						UE_LOG( LogHttpReplay, Error, TEXT( "FHttpNetworkReplayStreamer::HttpDownloadFinished. DecompressBuffer FAILED." ) );
+						SetLastError( ENetworkReplayError::ServiceUnavailable );
+						return;
+					}
+
+					StreamArchive.Buffer.Append( Uncompressed );
+				}
+				else
+				{
+					StreamArchive.Buffer.Append( HttpResponse->GetContent() );
+				}
+
 				StreamChunkIndex++;
 			}
 
@@ -1560,8 +1663,29 @@ void FHttpNetworkReplayStreamer::HttpDownloadCheckpointFinished( FHttpRequestPtr
 		}
 
 		// Get the checkpoint data
-		CheckpointArchive.Buffer	= HttpResponse->GetContent();
-		CheckpointArchive.Pos		= 0;
+		if ( SupportsCompression() )
+		{
+			SCOPE_CYCLE_COUNTER( STAT_HttpReplay_DecompressTime );
+
+			FHttpStreamFArchive Compressed;
+			Compressed.Buffer = HttpResponse->GetContent();
+
+			if ( !DecompressBuffer( Compressed, CheckpointArchive.Buffer ) )
+			{
+				UE_LOG( LogHttpReplay, Error, TEXT( "FHttpNetworkReplayStreamer::HttpDownloadCheckpointFinished. DecompressBuffer FAILED." ) );
+				GotoCheckpointDelegate.ExecuteIfBound( false, -1 );
+				GotoCheckpointDelegate	= FOnCheckpointReadyDelegate();
+				DownloadCheckpointIndex = -1;
+				LastGotoTimeInMS		= -1;
+				return;
+			}
+		}
+		else
+		{
+			CheckpointArchive.Buffer = HttpResponse->GetContent();
+		}
+
+		CheckpointArchive.Pos = 0;
 
 		// Completely reset our stream (we're going to start downloading from the start of the checkpoint)
 		StreamArchive.Buffer.Empty();
diff --git a/Engine/Source/Runtime/NetworkReplayStreaming/HttpNetworkReplayStreaming/Public/HttpNetworkReplayStreaming.h b/Engine/Source/Runtime/NetworkReplayStreaming/HttpNetworkReplayStreaming/Public/HttpNetworkReplayStreaming.h
--- a/Engine/Source/Runtime/NetworkReplayStreaming/HttpNetworkReplayStreaming/Public/HttpNetworkReplayStreaming.h
+++ b/Engine/Source/Runtime/NetworkReplayStreaming/HttpNetworkReplayStreaming/Public/HttpNetworkReplayStreaming.h
@@ -9,7 +9,7 @@
 /**
  * Archive used to buffer stream over http
  */
-class FHttpStreamFArchive : public FArchive
+class HTTPNETWORKREPLAYSTREAMING_API FHttpStreamFArchive : public FArchive
 {
 public:
 	FHttpStreamFArchive() : Pos( 0 ), bAtEndOfReplay( false ) {}
@@ -141,6 +141,22 @@ public:
 	FString		Meta;
 };
 
+/**
+* FQueuedHttpRequestAddUser
+* Custom event so that we can defer the need to knowing SessionName until we actually send it (which we should have it by then, since requests are executed in order)
+*/
+class FQueuedHttpRequestAddUser : public FQueuedHttpRequest
+{
+public:
+	FQueuedHttpRequestAddUser( const FString& InUser, TSharedRef< class IHttpRequest > InHttpRequest );
+
+	virtual ~FQueuedHttpRequestAddUser()
+	{
+	}
+
+	virtual bool PreProcess( FHttpNetworkReplayStreamer* Streamer, const FString& ServerURL, const FString& SessionName ) override;
+};
+
 /**
 * FQueuedGotoFakeCheckpoint
 */
@@ -192,6 +208,7 @@ public:
 	virtual void		SearchEvents(const FString& EventGroup, const FOnEnumerateStreamsComplete& Delegate) override;
 	virtual void		KeepReplay( const FString& ReplayName, const bool bKeep ) override;
 	virtual ENetworkReplayError::Type GetLastError() const override;
+	virtual FString		GetReplayID() const override { return SessionName; }
 
 	/** FHttpNetworkReplayStreamer */
 	void UploadHeader();
@@ -215,6 +232,9 @@ public:
 	void ConditionallyEnumerateCheckpoints();
 
 	virtual void ProcessRequestInternal( TSharedPtr< class IHttpRequest > Request );
+	virtual bool SupportsCompression() const { return false; }
+	virtual bool CompressBuffer( const TArray< uint8 >& InBuffer, FHttpStreamFArchive& OutCompressed ) const { return false; }
+	virtual bool DecompressBuffer(FHttpStreamFArchive& InCompressed, TArray< uint8 >& OutBuffer) const { return false; }
 
 	/** EStreamerState - Overall state of the streamer */
 	enum class EStreamerState
@@ -283,6 +303,8 @@ public:
 
 	TArray< TSharedPtr< FQueuedHttpRequest > >	QueuedHttpRequests;
 	TSharedPtr< FQueuedHttpRequest >			InFlightHttpRequest;
+
+	TSet< FString >					EventGroupSet;
 };
 
 class HTTPNETWORKREPLAYSTREAMING_API FHttpNetworkReplayStreamingFactory : public INetworkReplayStreamingFactory, public FTickableGameObject
diff --git a/Engine/Source/Runtime/NetworkReplayStreaming/InMemoryNetworkReplayStreaming/Public/InMemoryNetworkReplayStreaming.h b/Engine/Source/Runtime/NetworkReplayStreaming/InMemoryNetworkReplayStreaming/Public/InMemoryNetworkReplayStreaming.h
--- a/Engine/Source/Runtime/NetworkReplayStreaming/InMemoryNetworkReplayStreaming/Public/InMemoryNetworkReplayStreaming.h
+++ b/Engine/Source/Runtime/NetworkReplayStreaming/InMemoryNetworkReplayStreaming/Public/InMemoryNetworkReplayStreaming.h
@@ -77,6 +77,7 @@ public:
 	virtual void RequestEventData(const FString& EventID, const FOnRequestEventDataComplete& RequestEventDataComplete) override;
 	virtual void SearchEvents(const FString& EventGroup, const FOnEnumerateStreamsComplete& Delegate) override;
 	virtual void KeepReplay( const FString& ReplayName, const bool bKeep ) override;
+	virtual FString	GetReplayID() const override { return TEXT( "" ); }
 
 	/** FTickableObjectBase implementation */
 	virtual void Tick(float DeltaSeconds) override;
diff --git a/Engine/Source/Runtime/NetworkReplayStreaming/NetworkReplayStreaming/Public/NetworkReplayStreaming.h b/Engine/Source/Runtime/NetworkReplayStreaming/NetworkReplayStreaming/Public/NetworkReplayStreaming.h
--- a/Engine/Source/Runtime/NetworkReplayStreaming/NetworkReplayStreaming/Public/NetworkReplayStreaming.h
+++ b/Engine/Source/Runtime/NetworkReplayStreaming/NetworkReplayStreaming/Public/NetworkReplayStreaming.h
@@ -187,6 +187,7 @@ public:
 	virtual void KeepReplay( const FString& ReplayName, const bool bKeep ) = 0;
 	/** Returns true if the playing stream is currently in progress */
 	virtual bool IsLive() const = 0;
+	virtual FString	GetReplayID() const = 0;
 
 	/**
 	 * Attempts to delete the stream with the specified name. May execute asynchronously.
diff --git a/Engine/Source/Runtime/NetworkReplayStreaming/NullNetworkReplayStreaming/Public/NullNetworkReplayStreaming.h b/Engine/Source/Runtime/NetworkReplayStreaming/NullNetworkReplayStreaming/Public/NullNetworkReplayStreaming.h
--- a/Engine/Source/Runtime/NetworkReplayStreaming/NullNetworkReplayStreaming/Public/NullNetworkReplayStreaming.h
+++ b/Engine/Source/Runtime/NetworkReplayStreaming/NullNetworkReplayStreaming/Public/NullNetworkReplayStreaming.h
@@ -73,6 +73,7 @@ public:
 	virtual void RequestEventData(const FString& EventID, const FOnRequestEventDataComplete& RequestEventDataComplete) override;
 	virtual void SearchEvents(const FString& EventGroup, const FOnEnumerateStreamsComplete& Delegate) override;
 	virtual void KeepReplay( const FString& ReplayName, const bool bKeep ) override {}
+	virtual FString	GetReplayID() const override { return TEXT( "" ); }
 
 	/** FTickableObjectBase implementation */
 	virtual void Tick(float DeltaSeconds) override;
diff --git a/Engine/Source/Runtime/NullDrv/Private/NullRHI.cpp b/Engine/Source/Runtime/NullDrv/Private/NullRHI.cpp
--- a/Engine/Source/Runtime/NullDrv/Private/NullRHI.cpp
+++ b/Engine/Source/Runtime/NullDrv/Private/NullRHI.cpp
@@ -29,15 +29,19 @@ void FNullDynamicRHI::Init()
 	
 	check(!GIsRHIInitialized);
 
-	// Notify all initialized FRenderResources that there's a valid RHI device to create their RHI resources for now.
-	for(TLinkedList<FRenderResource*>::TIterator ResourceIt(FRenderResource::GetResourceList());ResourceIt;ResourceIt.Next())
+	// do not do this at least on dedicated server; clients with -NullRHI may need additional consideration
+    if (WITH_EDITOR || !IsRunningDedicatedServer())
 	{
-		ResourceIt->InitRHI();
-	}
-	// Dynamic resources can have dependencies on static resources (with uniform buffers) and must initialized last!
-	for(TLinkedList<FRenderResource*>::TIterator ResourceIt(FRenderResource::GetResourceList());ResourceIt;ResourceIt.Next())
-	{
-		ResourceIt->InitDynamicRHI();
+		// Notify all initialized FRenderResources that there's a valid RHI device to create their RHI resources for now.
+		for(TLinkedList<FRenderResource*>::TIterator ResourceIt(FRenderResource::GetResourceList());ResourceIt;ResourceIt.Next())
+		{
+			ResourceIt->InitRHI();
+		}
+		// Dynamic resources can have dependencies on static resources (with uniform buffers) and must initialized last!
+		for(TLinkedList<FRenderResource*>::TIterator ResourceIt(FRenderResource::GetResourceList());ResourceIt;ResourceIt.Next())
+		{
+			ResourceIt->InitDynamicRHI();
+		}
 	}
 
 	GIsRHIInitialized = true;
@@ -55,8 +59,10 @@ void FNullDynamicRHI::Shutdown()
  */
 void* FNullDynamicRHI::GetStaticBuffer()
 {
-	static void* Buffer = NULL;
-	if (!Buffer)
+    checkf(WITH_EDITOR || !IsRunningDedicatedServer(), TEXT("NullRHI should never allocate memory on the server. Change the caller to avoid doing allocs in when FApp::ShouldUseNullRHI() is true."));
+
+	static void* Buffer = nullptr;
+	if ((WITH_EDITOR || !IsRunningDedicatedServer()) && !Buffer)
 	{
 		// allocate an 64 meg buffer, should be big enough for any texture/surface
 		Buffer = FMemory::Malloc(64 * 1024 * 1024);
diff --git a/Engine/Source/Runtime/Online/OnlineFramework/Party/Private/Party.cpp b/Engine/Source/Runtime/Online/OnlineFramework/Party/Private/Party.cpp
--- a/Engine/Source/Runtime/Online/OnlineFramework/Party/Private/Party.cpp
+++ b/Engine/Source/Runtime/Online/OnlineFramework/Party/Private/Party.cpp
@@ -65,23 +65,34 @@ void UParty::OnLogoutComplete(int32 LocalUserNum, bool bWasSuccessful)
 	if (JoinedParties.Num())
 	{
 		UE_LOG(LogParty, Log, TEXT("Party cleanup on logout"));
+		TArray<FOnlinePartyTypeId> PartiesToRemove;
 		for (const auto& PartyKeyValue : JoinedParties)
 		{
-			UPartyGameState* Party = PartyKeyValue.Value;
-			if (Party && Party->GetPartyId().IsValid())
+			PartiesToRemove.Add(PartyKeyValue.Key);
+		}
+		for (const auto& PartyKey : PartiesToRemove)
+		{
+			UPartyGameState** FoundParty = JoinedParties.Find(PartyKey);
+			if (FoundParty)
 			{
-				if (Party->GetPartyId().IsValid())
-				{
-					UE_LOG(LogParty, Log, TEXT("[%s] Removed"), *(Party->GetPartyId()->ToDebugString()));
-				}
-				else
+				UPartyGameState* Party = *FoundParty;
+				if (Party && Party->GetPartyId().IsValid())
 				{
-					UE_LOG(LogParty, Log, TEXT("[%s] Removed - Invalid party Id"));
+					if (Party->GetPartyId().IsValid())
+					{
+						UE_LOG(LogParty, Log, TEXT("[%s] Removed"), *(Party->GetPartyId()->ToDebugString()));
+					}
+					else
+					{
+						UE_LOG(LogParty, Log, TEXT("[%s] Removed - Invalid party Id"));
+					}
+					Party->HandleRemovedFromParty(EMemberExitedReason::Left);
 				}
-				Party->HandleRemovedFromParty(EMemberExitedReason::Left);
+				JoinedParties.Remove(PartyKey);
 			}
 		}
 
+		ensure(JoinedParties.Num() == 0);
 		JoinedParties.Empty();
 	}
 }
@@ -116,6 +127,7 @@ void UParty::OnLoginStatusChanged(int32 LocalUserNum, ELoginStatus::Type OldStat
 			}
 
 			ensure(JoinedParties.Num() == 0);
+			JoinedParties.Empty();
 		}
 	}
 
@@ -799,11 +811,12 @@ void UParty::OnLeavePartyComplete(const FUniqueNetId& LocalUserId, const ELeaveP
 	InCompletionDelegate.ExecuteIfBound(LocalUserId, Result);
 }
 
-void UParty::GetDefaultPersistentPartySettings(EPartyType& PartyType, bool& bLeaderFriendsOnly, bool& bLeaderInvitesOnly)
+void UParty::GetDefaultPersistentPartySettings(EPartyType& PartyType, bool& bLeaderFriendsOnly, bool& bLeaderInvitesOnly, bool& bAllowInvites)
 {
 	PartyType = EPartyType::Public;
 	bLeaderInvitesOnly = false;
 	bLeaderFriendsOnly = false;
+	bAllowInvites = true;
 }
 
 void UParty::GetPersistentPartyConfiguration(FPartyConfiguration& PartyConfig)
@@ -811,7 +824,8 @@ void UParty::GetPersistentPartyConfiguration(FPartyConfiguration& PartyConfig)
 	EPartyType PartyType = EPartyType::Public;
 	bool bLeaderInvitesOnly = false;
 	bool bLeaderFriendsOnly = false;
-	GetDefaultPersistentPartySettings(PartyType, bLeaderFriendsOnly, bLeaderInvitesOnly);
+	bool bAllowInvites = true;
+	GetDefaultPersistentPartySettings(PartyType, bLeaderFriendsOnly, bLeaderInvitesOnly, bAllowInvites);
 
 	bool bIsPrivate = (PartyType == EPartyType::Private);
 
@@ -843,7 +857,7 @@ void UParty::GetPersistentPartyConfiguration(FPartyConfiguration& PartyConfig)
 	PartyConfig.bIsAcceptingMembers = bIsPrivate ? false : true;
 	PartyConfig.bShouldRemoveOnDisconnection = true;
 	PartyConfig.PresencePermissions = PresencePermissions;
-	PartyConfig.InvitePermissions = bLeaderInvitesOnly ? PartySystemPermissions::EInvitePermissions::Leader : PartySystemPermissions::EInvitePermissions::Anyone;
+	PartyConfig.InvitePermissions = bAllowInvites ? (bLeaderInvitesOnly ? PartySystemPermissions::EInvitePermissions::Leader : PartySystemPermissions::EInvitePermissions::Anyone) : PartySystemPermissions::EInvitePermissions::Noone;
 
 	PartyConfig.MaxMembers = DefaultMaxPartySize;
 }
@@ -899,9 +913,11 @@ void UParty::OnCreatePesistentPartyCompletedCommon(const FUniqueNetId& LocalUser
 	EPartyType PartyType = EPartyType::Public;
 	bool bLeaderInvitesOnly = false;
 	bool bLeaderFriendsOnly = false;
-	GetDefaultPersistentPartySettings(PartyType, bLeaderFriendsOnly, bLeaderInvitesOnly);
+	bool bAllowInvites = true;
+	GetDefaultPersistentPartySettings(PartyType, bLeaderFriendsOnly, bLeaderInvitesOnly, bAllowInvites);
 
 	PersistentParty->SetPartyType(PartyType, bLeaderFriendsOnly, bLeaderInvitesOnly);
+	PersistentParty->SetInvitesDisabled(!bAllowInvites);
 
 	TSharedPtr<const FUniqueNetId> PartyLeaderPtr = PersistentParty->GetPartyLeader();
 	ensure(PartyLeaderPtr.IsValid());
diff --git a/Engine/Source/Runtime/Online/OnlineFramework/Party/Private/PartyGameState.cpp b/Engine/Source/Runtime/Online/OnlineFramework/Party/Private/PartyGameState.cpp
--- a/Engine/Source/Runtime/Online/OnlineFramework/Party/Private/PartyGameState.cpp
+++ b/Engine/Source/Runtime/Online/OnlineFramework/Party/Private/PartyGameState.cpp
@@ -570,7 +570,15 @@ void UPartyGameState::HandleLockoutPromotionStateChange(bool bNewLockoutState)
 
 EApprovalAction UPartyGameState::ProcessJoinRequest(const FUniqueNetId& RecipientId, const FUniqueNetId& SenderId, EJoinPartyDenialReason& DenialReason)
 {
-	return EApprovalAction::Approve;
+	if (IsInJoinableGameState())
+	{
+		return EApprovalAction::Approve;
+	}
+	else
+	{
+		DenialReason = EJoinPartyDenialReason::GameFull;
+		return EApprovalAction::Deny;
+	}
 }
 
 void UPartyGameState::HandlePartyJoinRequestReceived(const FUniqueNetId& RecipientId, const FUniqueNetId& SenderId)
@@ -594,15 +602,8 @@ void UPartyGameState::HandlePartyJoinRequestReceived(const FUniqueNetId& Recipie
 		int32 MaxPartyMembers = CurrentConfig.MaxMembers;
 		if (NumPartyMembers < MaxPartyMembers)
 		{
-			if (IsInJoinableGameState())
-			{
-				ApprovalAction = ProcessJoinRequest(RecipientId, SenderId, DenialReason);
-			}
-			else
-			{
-				// Game is full
-				DenialReason = EJoinPartyDenialReason::GameFull;
-			}
+			// Give the game a chance to accept or deny this player
+			ApprovalAction = ProcessJoinRequest(RecipientId, SenderId, DenialReason);
 		}
 		else
 		{
diff --git a/Engine/Source/Runtime/Online/OnlineFramework/Party/Public/Party.h b/Engine/Source/Runtime/Online/OnlineFramework/Party/Public/Party.h
--- a/Engine/Source/Runtime/Online/OnlineFramework/Party/Public/Party.h
+++ b/Engine/Source/Runtime/Online/OnlineFramework/Party/Public/Party.h
@@ -352,8 +352,9 @@ protected:
 	 * @param PartyType party visibility setting
 	 * @param bLeaderFriendsOnly can only the leader send invites
 	 * @param bLeaderInvitesOnly can only the leader have friends join via presence
+	 * @param bAllowInvites are invites currently allowed
 	 */
-	virtual void GetDefaultPersistentPartySettings(EPartyType& PartyType, bool& bLeaderFriendsOnly, bool& bLeaderInvitesOnly);
+	virtual void GetDefaultPersistentPartySettings(EPartyType& PartyType, bool& bLeaderFriendsOnly, bool& bLeaderInvitesOnly, bool& bAllowInvites);
 
 	/**
 	 * Get the default configuration used for persistent party creation
diff --git a/Engine/Source/Runtime/Online/OnlineFramework/Party/Public/PartyGameState.h b/Engine/Source/Runtime/Online/OnlineFramework/Party/Public/PartyGameState.h
--- a/Engine/Source/Runtime/Online/OnlineFramework/Party/Public/PartyGameState.h
+++ b/Engine/Source/Runtime/Online/OnlineFramework/Party/Public/PartyGameState.h
@@ -314,6 +314,9 @@ public:
 		}
 	}
 
+	/** @return Unique ID of the user who created or joined this room (not the party leader) */
+	const FUniqueNetIdRepl& GetOwningUserId() const { return OwningUserId; }
+
 	/** @return delegate fired when global party data has changed */
 	FOnPartyDataChanged& OnPartyDataChanged() { return PartyDataChanged; }
 	/** @return delegate fired when an individual party member's data has changed */
diff --git a/Engine/Source/Runtime/Online/OnlineSubsystem/Public/Interfaces/OnlineExternalUIInterface.h b/Engine/Source/Runtime/Online/OnlineSubsystem/Public/Interfaces/OnlineExternalUIInterface.h
--- a/Engine/Source/Runtime/Online/OnlineSubsystem/Public/Interfaces/OnlineExternalUIInterface.h
+++ b/Engine/Source/Runtime/Online/OnlineSubsystem/Public/Interfaces/OnlineExternalUIInterface.h
@@ -86,7 +86,8 @@ struct FShowWebUrlParams
  */
 enum class EPlatformMessageType
 {
-	EmptyStore
+	EmptyStore,
+	ChatRestricted
 };
 
 /** 
diff --git a/Engine/Source/Runtime/Online/OnlineSubsystemUtils/Classes/IpNetDriver.h b/Engine/Source/Runtime/Online/OnlineSubsystemUtils/Classes/IpNetDriver.h
--- a/Engine/Source/Runtime/Online/OnlineSubsystemUtils/Classes/IpNetDriver.h
+++ b/Engine/Source/Runtime/Online/OnlineSubsystemUtils/Classes/IpNetDriver.h
@@ -70,4 +70,9 @@ class ONLINESUBSYSTEMUTILS_API UIpNetDriver : public UNetDriver
 
 	/** @return TCPIP connection to server */
 	class UIpConnection* GetServerConnection();
+
+	// Callback for platform handling when networking is taking a long time in a single frame (by default over 1 second).
+	// It may get called multiple times in a single frame if additional processing after a previous alert exceeds the threshold again
+	DECLARE_MULTICAST_DELEGATE(FOnNetworkProcessingCausingSlowFrame);
+	static FOnNetworkProcessingCausingSlowFrame OnNetworkProcessingCausingSlowFrame;
 };
diff --git a/Engine/Source/Runtime/Online/OnlineSubsystemUtils/Private/IpNetDriver.cpp b/Engine/Source/Runtime/Online/OnlineSubsystemUtils/Private/IpNetDriver.cpp
--- a/Engine/Source/Runtime/Online/OnlineSubsystemUtils/Private/IpNetDriver.cpp
+++ b/Engine/Source/Runtime/Online/OnlineSubsystemUtils/Private/IpNetDriver.cpp
@@ -17,6 +17,26 @@ Notes:
 	Declarations.
 -----------------------------------------------------------------------------*/
 
+UIpNetDriver::FOnNetworkProcessingCausingSlowFrame UIpNetDriver::OnNetworkProcessingCausingSlowFrame;
+
+// Time before the alarm delegate is called (in seconds)
+float GIpNetDriverMaxDesiredTimeSliceBeforeAlarmSecs = 1.0f;
+
+FAutoConsoleVariableRef GIpNetDriverMaxDesiredTimeSliceBeforeAlarmSecsCVar(
+	TEXT("n.IpNetDriverMaxFrameTimeBeforeAlert"),
+	GIpNetDriverMaxDesiredTimeSliceBeforeAlarmSecs,
+	TEXT("Time to spend processing networking data in a single frame before an alert is raised (in seconds)\n")
+	TEXT("It may get called multiple times in a single frame if additional processing after a previous alert exceeds the threshold again\n")
+	TEXT(" default: 1 s"));
+
+// Time before the time taken in a single frame is printed out (in seconds)
+float GIpNetDriverLongFramePrintoutThresholdSecs = 10.0f;
+
+FAutoConsoleVariableRef GIpNetDriverLongFramePrintoutThresholdSecsCVar(
+	TEXT("n.IpNetDriverMaxFrameTimeBeforeLogging"),
+	GIpNetDriverLongFramePrintoutThresholdSecs,
+	TEXT("Time to spend processing networking data in a single frame before an output log warning is printed (in seconds)\n")
+	TEXT(" default: 10 s"));
 
 UIpNetDriver::UIpNetDriver(const FObjectInitializer& ObjectInitializer)
 	: Super(ObjectInitializer)
@@ -169,12 +189,23 @@ void UIpNetDriver::TickDispatch( float DeltaTime )
 	ISocketSubsystem* SocketSubsystem = GetSocketSubsystem();
 
 	const double StartReceiveTime = FPlatformTime::Seconds();
+	double AlarmTime = StartReceiveTime + GIpNetDriverMaxDesiredTimeSliceBeforeAlarmSecs;
 
 	// Process all incoming packets.
 	uint8 Data[MAX_PACKET_SIZE];
 	TSharedRef<FInternetAddr> FromAddr = SocketSubsystem->CreateInternetAddr();
 	for( ; Socket != NULL; )
 	{
+		{
+			const double CurrentTime = FPlatformTime::Seconds();
+			if (CurrentTime > AlarmTime)
+			{
+				OnNetworkProcessingCausingSlowFrame.Broadcast();
+
+				AlarmTime = CurrentTime + GIpNetDriverMaxDesiredTimeSliceBeforeAlarmSecs;
+			}
+		}
+
 		int32 BytesRead = 0;
 		// Get data, if any.
 		CLOCK_CYCLES(RecvCycles);
@@ -306,11 +337,10 @@ void UIpNetDriver::TickDispatch( float DeltaTime )
 		}
 	}
 
-	const double EndReceiveTime		= FPlatformTime::Seconds();
-	const float DeltaReceiveTime	= EndReceiveTime - StartReceiveTime;
-	const float Threshold			= 10.0f;
+	const double EndReceiveTime = FPlatformTime::Seconds();
+	const float DeltaReceiveTime = EndReceiveTime - StartReceiveTime;
 
-	if ( DeltaReceiveTime > Threshold )
+	if (DeltaReceiveTime > GIpNetDriverLongFramePrintoutThresholdSecs)
 	{
 		UE_LOG( LogNet, Warning, TEXT( "UIpNetDriver::TickDispatch: Took too long to receive packets. Time: %2.2f %s" ), DeltaReceiveTime, *GetName() );
 	}
diff --git a/Engine/Source/Runtime/Online/XMPP/Private/XmppConnection.cpp b/Engine/Source/Runtime/Online/XMPP/Private/XmppConnection.cpp
--- a/Engine/Source/Runtime/Online/XMPP/Private/XmppConnection.cpp
+++ b/Engine/Source/Runtime/Online/XMPP/Private/XmppConnection.cpp
@@ -3,13 +3,13 @@
 #include "XmppPrivatePCH.h"
 #include "XmppConnection.h"
 
-bool FXmppUserJid::ParseResource(FString& OutAppId, FString& OutPlatform) const
+bool FXmppUserJid::ParseResource(const FString& InResource, FString& OutAppId, FString& OutPlatform)
 {
 	OutAppId.Empty();
 	OutPlatform.Empty();
 
 	TArray<FString> ParsedResource;
-	if (Resource.ParseIntoArray(ParsedResource, TEXT(":"), false) > 1)
+	if (InResource.ParseIntoArray(ParsedResource, TEXT(":"), false) > 1)
 	{
 		if (ParsedResource[0].StartsWith(TEXT("V")))
 		{
@@ -28,7 +28,7 @@ bool FXmppUserJid::ParseResource(FString& OutAppId, FString& OutPlatform) const
 	else
 	{
 		FString ClientId, NotUsed;
-		if (Resource.Split(TEXT("-"), &ClientId, &NotUsed) &&
+		if (InResource.Split(TEXT("-"), &ClientId, &NotUsed) &&
 			!ClientId.IsEmpty())
 		{
 			OutAppId = ClientId;
@@ -37,3 +37,4 @@ bool FXmppUserJid::ParseResource(FString& OutAppId, FString& OutPlatform) const
 	}
 	return false;
 }
+
diff --git a/Engine/Source/Runtime/Online/XMPP/Public/XmppConnection.h b/Engine/Source/Runtime/Online/XMPP/Public/XmppConnection.h
--- a/Engine/Source/Runtime/Online/XMPP/Public/XmppConnection.h
+++ b/Engine/Source/Runtime/Online/XMPP/Public/XmppConnection.h
@@ -77,12 +77,26 @@ public:
 	/**
 	 * Get the components that comprise the resource
 	 *
+	 * @param InResource The resource to parse
 	 * @param OutAppId The app id the user is using
 	 * @param OutPlatform The platform the user is using
 	 *
 	 * @return Whether the Resource was successfully parsed or not
 	 */
-	bool ParseResource(FString& OutAppId, FString& OutPlatform) const;
+	static bool ParseResource(const FString& InResource, FString& OutAppId, FString& OutPlatform);
+
+	/** 
+	 * Get the components that comprise the resource
+	 *
+	 * @param OutAppId The app id the user is using
+	 * @param OutPlatform The platform the user is using
+	 *
+	 * @return Whether the Resource was successfully parsed or not
+	 */
+	bool ParseResource(FString& OutAppId, FString& OutPlatform) const
+	{
+		return ParseResource(Resource, OutAppId, OutPlatform);
+	}
 
 	bool operator==(const FXmppUserJid& Other)
 	{
diff --git a/Engine/Source/Runtime/RHI/Private/DynamicRHI.cpp b/Engine/Source/Runtime/RHI/Private/DynamicRHI.cpp
--- a/Engine/Source/Runtime/RHI/Private/DynamicRHI.cpp
+++ b/Engine/Source/Runtime/RHI/Private/DynamicRHI.cpp
@@ -6,6 +6,7 @@
 
 #include "RHI.h"
 #include "ModuleManager.h"
+#include "GenericPlatformDriver.h" // FGPUDriverInfo
 
 #ifndef PLATFORM_ALLOW_NULL_RHI
 	#define PLATFORM_ALLOW_NULL_RHI		0
@@ -14,14 +15,17 @@
 // Globals.
 FDynamicRHI* GDynamicRHI = NULL;
 
-static TAutoConsoleVariable<int32> CVarDetectAndWarnOfBadDrivers(
-	TEXT("r.DetectAndWarnOfBadDrivers"),
+static TAutoConsoleVariable<int32> CVarWarnOfBadDrivers(
+	TEXT("r.WarnOfBadDrivers"),
 	1,
 	TEXT("On engine startup we can check the current GPU driver and warn the user about issues and suggest a specific version\n")
 	TEXT("The test is fast so this should not cost any performance.\n")
 	TEXT(" 0: off\n")
 	TEXT(" 1: a message on startup might appear (default)\n")
-	TEXT(" 2: Test by pretending the driver has issues"),
+	TEXT(" 2: Simulating the system has a blacklisted NVIDIA driver (UI should appear)\n")
+	TEXT(" 3: Simulating the system has a blacklisted AMD driver (UI should appear)\n")
+	TEXT(" 4: Simulating the system has a not blacklisted AMD driver (no UI should appear)\n")
+	TEXT(" 5: Simulating the system has a Indel driver (no UI should appear)"),
 	ECVF_RenderThreadSafe
 	);
 
@@ -45,40 +49,10 @@ void InitNullRHI()
 	GRHISupportsTextureStreaming = false;
 }
 
-// e.g. 36143 for the NVIDIA driver 361.43
-// no floating point to avoid precision issues
-// @return -1 if the version is unknown
-static int32 ExtractGPUDriverVersionInt(const FGPUDriverInfo& DriverInfo)
-{
-	int32 Ret = -1;
-
-	// we use the internal version, not the user version to avoid problem where the name was altered 
-	const FString& InternalVersion = DriverInfo.InternalDriverVersion;
-
-	if(DriverInfo.IsNVIDIA())
-	{
-		// we don't care about the windows version so we don't look at the front part of the driver version
-		// e.g. 36.143
-		FString RightPart = InternalVersion.Right(6);
-		// e.g. 36143
-		FString CompactNumber = RightPart.Replace(TEXT("."), TEXT(""));
-
-		if(FCString::IsNumeric(*CompactNumber))
-		{
-			Ret = FCString::Atoi(*CompactNumber);
-		}
-	}
-	else
-	{
-		// examples for AMD: "13.12" "15.101.1007" "13.351"
-		// todo: need to decide on what part we extract, maybe make the return argument multiple ints
-	}
-	return Ret;
-}
 
 void RHIDetectAndWarnOfBadDrivers()
 {
-	int32 CVarValue = CVarDetectAndWarnOfBadDrivers.GetValueOnGameThread();
+	int32 CVarValue = CVarWarnOfBadDrivers.GetValueOnGameThread();
 
 	if(!GIsRHIInitialized || !CVarValue || GRHIVendorId == 0)
 	{
@@ -95,62 +69,69 @@ void RHIDetectAndWarnOfBadDrivers()
 	DriverInfo.UserDriverVersion = GRHIAdapterUserDriverVersion;
 	DriverInfo.DriverDate = GRHIAdapterDriverDate;
 
+
+#if !(UE_BUILD_SHIPPING || UE_BUILD_TEST)
 	// for testing
 	if(CVarValue == 2)
 	{
-		DriverInfo.VendorId = 0x10DE;
-		DriverInfo.DeviceDescription = TEXT("Test NVIDIA");
+		DriverInfo.SetNVIDIA();
+		DriverInfo.DeviceDescription = TEXT("Test NVIDIA (bad)");
 		DriverInfo.UserDriverVersion = TEXT("361.43");
 		DriverInfo.InternalDriverVersion = TEXT("9.18.136.143");
 		DriverInfo.DriverDate = TEXT("01-01-1900");
 	}
-
-	int32 VersionInt = ExtractGPUDriverVersionInt(DriverInfo);
-
-	if(VersionInt < 0)
+	else if(CVarValue == 3)
 	{
-		// no valid version
-		return;
+		DriverInfo.SetAMD();
+		DriverInfo.DeviceDescription = TEXT("Test AMD (bad)");
+		DriverInfo.UserDriverVersion = TEXT("Test Catalyst Version");
+		DriverInfo.InternalDriverVersion = TEXT("13.152.1.1000");
+		DriverInfo.DriverDate = TEXT("09-10-13");
+	}
+	else if(CVarValue == 4)
+	{
+		DriverInfo.SetAMD();
+		DriverInfo.DeviceDescription = TEXT("Test AMD (good)");
+		DriverInfo.UserDriverVersion = TEXT("Test Catalyst Version");
+		DriverInfo.InternalDriverVersion = TEXT("15.30.1025.1001");
+		DriverInfo.DriverDate = TEXT("01-01-16");
 	}
+	else if(CVarValue == 5)
+	{
+		DriverInfo.SetIntel();
+		DriverInfo.DeviceDescription = TEXT("Test Intel (good)");
+		DriverInfo.UserDriverVersion = TEXT("Test Intel Version");
+		DriverInfo.InternalDriverVersion = TEXT("8.15.10.2302");
+		DriverInfo.DriverDate = TEXT("01-01-15");
+	}
+#endif // !(UE_BUILD_SHIPPING || UE_BUILD_TEST)
+
+	FGPUHardware DetectedGPUHardware(DriverInfo);
 
-	// can be made data driven:
-	if(DriverInfo.IsNVIDIA())
+	if (DriverInfo.IsValid())
 	{
-		// todo: make this data driven
-		FString SuggestedVersion = TEXT("359.06");
+		FBlackListEntry BlackListEntry = DetectedGPUHardware.FindDriverBlacklistEntry();
 
-		// todo: make this data driven
-		if(VersionInt == 36143	// 361.43 "UE-25096 Viewport flashes black and white when moving in the scene on latest Nvidia drivers", at that time that was the latest driver
-		|| VersionInt <= 34709)	// 347.09 "old NVIDIA driver version we've seen driver crashes with Paragon content", we might have to adjust that further
+		if (BlackListEntry.IsValid())
 		{
+			// format message box UI
 			FFormatNamedArguments Args;
 			Args.Add(TEXT("AdapterName"), FText::FromString(DriverInfo.DeviceDescription));
 			Args.Add(TEXT("UserDriverVersion"), FText::FromString(DriverInfo.UserDriverVersion));
 			Args.Add(TEXT("InternalDriverVersion"), FText::FromString(DriverInfo.InternalDriverVersion));
 			Args.Add(TEXT("DriverDate"), FText::FromString(DriverInfo.DriverDate));
-			Args.Add(TEXT("SuggestedVersion"), FText::FromString(SuggestedVersion));
+			Args.Add(TEXT("SuggestedVersion"), FText::FromString(DetectedGPUHardware.GetSuggestedDriverVersion()));
 
+			// this message can be suppressed with r.WarnOfBadDrivers=0
 			FText LocalizedMsg = FText::Format(NSLOCTEXT("MessageDialog", "VideoCardDriverIssueReport",
-				"Your current video card driver version might cause problems.\n(reduced performance, artifacts, stalls, crashes, restarts)\n\nYour current driver:\n  Name:\t{AdapterName}\n  Version:\t{UserDriverVersion} (internal {InternalDriverVersion})\n  Date:\t{DriverDate}\n\nSuggested driver:\n  Version:\t{SuggestedVersion}\n\n"),
+				"Your current video card driver version might cause problems.\n(reduced performance, artifacts, stalls, crashes, restarts)\n\nYour current driver:\n  Name:\t{AdapterName}\n  Version:\t{UserDriverVersion} (internal {InternalDriverVersion})\n  Date:\t{DriverDate}\n\nSuggested driver:\n  Version:\t{SuggestedVersion}"),
 				Args);
 
-			FPlatformMisc::MessageBoxExt(EAppMsgType::Ok, 
+			FPlatformMisc::MessageBoxExt(EAppMsgType::Ok,
 				*LocalizedMsg.ToString(),
 				*NSLOCTEXT("MessageDialog", "TitleVideoCardDriverIssue", "WARNING: Video card driver").ToString());
 		}
 	}
-	else if(DriverInfo.IsAMD())
-	{
-		// so far we don't have any data on issues
-	}
-	else if(DriverInfo.IsIntel())
-	{
-		// so far we don't have any data on issues
-	}
-	else
-	{
-		// so far we don't have any data on issues
-	}
 }
 
 void RHIInit(bool bHasEditorToken)
