Copying //UE4/Orion-Staging to //UE4/Main (Origin //Orion/Dev-General @ 2861092)

diff --git a/Engine/Source/Runtime/Engine/Private/Particles/ParticleEmitterInstances.cpp b/Engine/Source/Runtime/Engine/Private/Particles/ParticleEmitterInstances.cpp
--- a/Engine/Source/Runtime/Engine/Private/Particles/ParticleEmitterInstances.cpp
+++ b/Engine/Source/Runtime/Engine/Private/Particles/ParticleEmitterInstances.cpp
@@ -328,6 +328,7 @@ FParticleEmitterInstance::FParticleEmitterInstance() :
 	, bIgnoreComponentScale(0)
 	, bIsBeam(0)
 	, bAxisLockEnabled(0)
+	, bFakeBurstsWhenSpawningSupressed(0)
 	, LockAxisFlags(EPAL_NONE)
 	, SortMode(PSORTMODE_None)
     , ParticleData(NULL)
@@ -708,45 +709,51 @@ void FParticleEmitterInstance::Tick(float DeltaTime, bool bSuppressSpawning)
 	// Handle EmitterTime setup, looping, etc.
 	float EmitterDelay = Tick_EmitterTimeSetup(DeltaTime, LODLevel);
 
-	// Kill off any dead particles
-	KillParticles();
-
-	// Reset particle parameters.
-	ResetParticleParameters(DeltaTime);
+	if (bEnabled)
+	{
+		// Kill off any dead particles
+		KillParticles();
 
-	// Update the particles
-	SCOPE_CYCLE_COUNTER(STAT_SpriteUpdateTime);
-	CurrentMaterial = LODLevel->RequiredModule->Material;
-	Tick_ModuleUpdate(DeltaTime, LODLevel);
+		// Reset particle parameters.
+		ResetParticleParameters(DeltaTime);
 
-	// Spawn new particles.
-	SpawnFraction = Tick_SpawnParticles(DeltaTime, LODLevel, bSuppressSpawning, bFirstTime);
+		// Update the particles
+		SCOPE_CYCLE_COUNTER(STAT_SpriteUpdateTime);
+		CurrentMaterial = LODLevel->RequiredModule->Material;
+		Tick_ModuleUpdate(DeltaTime, LODLevel);
 
-	// PostUpdate (beams only)
-	Tick_ModulePostUpdate(DeltaTime, LODLevel);
+		// Spawn new particles.
+		SpawnFraction = Tick_SpawnParticles(DeltaTime, LODLevel, bSuppressSpawning, bFirstTime);
 
-	if (ActiveParticles > 0)
-	{
-		// Update the orbit data...
-		UpdateOrbitData(DeltaTime);
-		// Calculate bounding box and simulate velocity.
-		UpdateBoundingBox(DeltaTime);
-	}
+		// PostUpdate (beams only)
+		Tick_ModulePostUpdate(DeltaTime, LODLevel);
 
-	Tick_ModuleFinalUpdate(DeltaTime, LODLevel);
+		if (ActiveParticles > 0)
+		{
+			// Update the orbit data...
+			UpdateOrbitData(DeltaTime);
+			// Calculate bounding box and simulate velocity.
+			UpdateBoundingBox(DeltaTime);
+		}
 
+		Tick_ModuleFinalUpdate(DeltaTime, LODLevel);
 
-	CheckEmitterFinished();
+		CheckEmitterFinished();
 
-	// Invalidate the contents of the vertex/index buffer.
-	IsRenderDataDirty = 1;
+		// Invalidate the contents of the vertex/index buffer.
+		IsRenderDataDirty = 1;
+	}
+	else
+	{
+		FakeBursts();
+	}
 
 	// 'Reset' the emitter time so that the delay functions correctly
 	EmitterTime += EmitterDelay;
 
 	// Store the last delta time.
 	LastDeltaTime = DeltaTime;
-	
+
 	// Reset particles position offset
 	PositionOffsetThisTick = FVector::ZeroVector;
 
@@ -931,6 +938,10 @@ float FParticleEmitterInstance::Tick_SpawnParticles(float DeltaTime, UParticleLO
 			SpawnFraction = Spawn(DeltaTime);
 		}
 	}
+	else if (bFakeBurstsWhenSpawningSupressed)
+	{
+		FakeBursts();
+	}
 	
 	return SpawnFraction;
 }
@@ -1089,6 +1100,7 @@ void FParticleEmitterInstance::Rewind()
 	EmitterTime = 0;
 	LoopCount = 0;
 	ParticleCounter = 0;
+	bEnabled = 1;
 	ResetBurstList();
 }
 
@@ -2277,6 +2289,31 @@ void FParticleEmitterInstance::KillParticle(int32 Index)
 	}
 }
 
+void FParticleEmitterInstance::FakeBursts()
+{
+	UParticleLODLevel* LODLevel = GetCurrentLODLevelChecked();
+	if (LODLevel->SpawnModule->BurstList.Num() > 0)
+	{
+		// For each burst in the list
+		for (int32 BurstIdx = 0; BurstIdx < LODLevel->SpawnModule->BurstList.Num(); BurstIdx++)
+		{
+			FParticleBurst* BurstEntry = &(LODLevel->SpawnModule->BurstList[BurstIdx]);
+			// If it hasn't been fired
+			if (LODLevel->Level < BurstFired.Num())
+			{
+				FLODBurstFired& LocalBurstFired = BurstFired[LODLevel->Level];
+				if (BurstIdx < LocalBurstFired.Fired.Num())
+				{
+					if (EmitterTime >= BurstEntry->Time)
+					{
+						LocalBurstFired.Fired[BurstIdx] = true;
+					}
+				}
+			}
+		}
+	}
+}
+
 /**
  *	This is used to force "kill" particles irrespective of their duration.
  *	Basically, this takes all particles and moves them to the 'end' of the 
@@ -2505,14 +2542,14 @@ bool FParticleEmitterInstance::FillReplayData( FDynamicEmitterReplayDataBase& Ou
 	// NOTE: This the base class implementation that should ONLY be called by derived classes' FillReplayData()!
 
 	// Make sure there is a template present
-	if (!SpriteTemplate || !bEnabled)
+	if (!SpriteTemplate)
 	{
 		return false;
 	}
 
 	// Allocate it for now, but we will want to change this to do some form
 	// of caching
-	if (ActiveParticles <= 0)
+	if (ActiveParticles <= 0 || !bEnabled)
 	{
 		return false;
 	}
@@ -2732,7 +2769,7 @@ FDynamicEmitterDataBase* FParticleSpriteEmitterInstance::GetDynamicData(bool bSe
 
 	// It is valid for the LOD level to be NULL here!
 	UParticleLODLevel* LODLevel = SpriteTemplate->GetCurrentLODLevel(this);
-	if (IsDynamicDataRequired(LODLevel) == false)
+	if (IsDynamicDataRequired(LODLevel) == false || !bEnabled)
 	{
 		return NULL;
 	}
@@ -2771,7 +2808,7 @@ bool FParticleSpriteEmitterInstance::UpdateDynamicData(FDynamicEmitterDataBase*
 
 	checkf((DynamicData->GetSource().eEmitterType == DET_Sprite), TEXT("Sprite::UpdateDynamicData> Invalid DynamicData type!"));
 
-	if (ActiveParticles <= 0)
+	if (ActiveParticles <= 0 || !bEnabled)
 	{
 		return false;
 	}
@@ -2801,7 +2838,7 @@ bool FParticleSpriteEmitterInstance::UpdateDynamicData(FDynamicEmitterDataBase*
  */
 FDynamicEmitterReplayDataBase* FParticleSpriteEmitterInstance::GetReplayData()
 {
-	if (ActiveParticles <= 0)
+	if (ActiveParticles <= 0 || !bEnabled)
 	{
 		return NULL;
 	}
@@ -2998,7 +3035,7 @@ void FParticleMeshEmitterInstance::Tick(float DeltaTime, bool bSuppressSpawning)
 
 	UParticleLODLevel* LODLevel = GetCurrentLODLevelChecked();
 	// See if we are handling mesh rotation
-	if (MeshRotationActive)
+	if (MeshRotationActive && bEnabled)
 	{
 		// Update the rotation for each particle
 		for (int32 i = 0; i < ActiveParticles; i++)
@@ -3067,7 +3104,7 @@ void FParticleMeshEmitterInstance::Tick(float DeltaTime, bool bSuppressSpawning)
 	// Call the standard tick
 	FParticleEmitterInstance::Tick(DeltaTime, bSuppressSpawning);
 	
-	if (MeshRotationActive)
+	if (MeshRotationActive && bEnabled)
 	{
 		//Must do this (at least) after module update other wise the reset value of RotationRate is used.
 		//Probably the other stuff before the module tick should be brought down here too and just leave the RotationRate reset before.
@@ -3323,7 +3360,7 @@ FDynamicEmitterDataBase* FParticleMeshEmitterInstance::GetDynamicData(bool bSele
 
 	// It is safe for LOD level to be NULL here!
 	UParticleLODLevel* LODLevel = SpriteTemplate->GetCurrentLODLevel(this);
-	if (IsDynamicDataRequired(LODLevel) == false)
+	if (IsDynamicDataRequired(LODLevel) == false || !bEnabled)
 	{
 		return NULL;
 	}
@@ -3365,7 +3402,7 @@ bool FParticleMeshEmitterInstance::UpdateDynamicData(FDynamicEmitterDataBase* Dy
 {
 	QUICK_SCOPE_CYCLE_COUNTER(STAT_ParticleMeshEmitterInstance_UpdateDynamicData);
 
-	if (ActiveParticles <= 0)
+	if (ActiveParticles <= 0 || !bEnabled)
 	{
 		return false;
 	}
@@ -3402,7 +3439,7 @@ bool FParticleMeshEmitterInstance::UpdateDynamicData(FDynamicEmitterDataBase* Dy
  */
 FDynamicEmitterReplayDataBase* FParticleMeshEmitterInstance::GetReplayData()
 {
-	if (ActiveParticles <= 0)
+	if (ActiveParticles <= 0 || !bEnabled)
 	{
 		return NULL;
 	}
diff --git a/Engine/Source/Runtime/Engine/Private/Particles/ParticleGpuSimulation.cpp b/Engine/Source/Runtime/Engine/Private/Particles/ParticleGpuSimulation.cpp
--- a/Engine/Source/Runtime/Engine/Private/Particles/ParticleGpuSimulation.cpp
+++ b/Engine/Source/Runtime/Engine/Private/Particles/ParticleGpuSimulation.cpp
@@ -2499,6 +2499,9 @@ public:
 	bool bReleased_GameThread;
 	bool bDestroyed_GameThread;
 
+	/** Allows disabling of simulation. */
+	bool bEnabled;
+
 	/** Default constructor. */
 	FParticleSimulationGPU()
 		: EmitterSimulationResources(NULL)
@@ -2510,6 +2513,7 @@ public:
 		, bDirty_GameThread(true)
 		, bReleased_GameThread(true)
 		, bDestroyed_GameThread(false)
+		, bEnabled(true)
 	{
 	}
 
@@ -3186,7 +3190,7 @@ public:
 
 		// Grab the current LOD level
 		UParticleLODLevel* LODLevel = GetCurrentLODLevelChecked();
-		if (LODLevel->bEnabled == false)
+		if (LODLevel->bEnabled == false || !bEnabled)
 		{
 			return NULL;
 		}
@@ -3483,139 +3487,153 @@ public:
 		// Handle EmitterTime setup, looping, etc.
 		float EmitterDelay = Tick_EmitterTimeSetup( DeltaSeconds, LODLevel );
 
-		// If the emitter is warming up but any particle spawned now will die
-		// anyway, suppress spawning.
-		if (Component && Component->bWarmingUp &&
-			Component->WarmupTime - SecondsSinceCreation > EmitterInfo.MaxLifetime)
+		Simulation->bEnabled = bEnabled;
+		if (bEnabled)
 		{
-			bSuppressSpawning = true;
-		}
-
-		// Mark any tiles with all dead particles as free.
-		int32 ActiveTileCount = MarkTilesInactive();
+			// If the emitter is warming up but any particle spawned now will die
+			// anyway, suppress spawning.
+			if (Component && Component->bWarmingUp &&
+				Component->WarmupTime - SecondsSinceCreation > EmitterInfo.MaxLifetime)
+			{
+				bSuppressSpawning = true;
+			}
 
-		// Update modules
-		Tick_ModuleUpdate(DeltaSeconds, LODLevel);
+			// Mark any tiles with all dead particles as free.
+			int32 ActiveTileCount = MarkTilesInactive();
 
-		// Spawn particles.
-		bool bRefreshTiles = false;
-		const bool bPreventSpawning = bHaltSpawning || bSuppressSpawning;
-		const bool bValidEmitterTime = (EmitterTime >= 0.0f);
-		const bool bValidLoop = AllowedLoopCount == 0 || LoopCount < AllowedLoopCount;
-		if (!bPreventSpawning && bValidEmitterTime && bValidLoop)
-		{
-			SCOPE_CYCLE_COUNTER(STAT_GPUSpriteSpawnTime);
+			// Update modules
+			Tick_ModuleUpdate(DeltaSeconds, LODLevel);
 
-			// Determine burst count.
-			FSpawnInfo BurstInfo;
-			int32 LeftoverBurst = 0;
+			// Spawn particles.
+			bool bRefreshTiles = false;
+			const bool bPreventSpawning = bHaltSpawning || bSuppressSpawning;
+			const bool bValidEmitterTime = (EmitterTime >= 0.0f);
+			const bool bValidLoop = AllowedLoopCount == 0 || LoopCount < AllowedLoopCount;
+			if (!bPreventSpawning && bValidEmitterTime && bValidLoop)
 			{
-				float BurstDeltaTime = DeltaSeconds;
-				GetCurrentBurstRateOffset(BurstDeltaTime, BurstInfo.Count);
+				SCOPE_CYCLE_COUNTER(STAT_GPUSpriteSpawnTime);
 
-				BurstInfo.Count += ForceBurstSpawnedParticles.Num();
-
-				if (BurstInfo.Count > FXConsoleVariables::MaxGPUParticlesSpawnedPerFrame)
+				// Determine burst count.
+				FSpawnInfo BurstInfo;
+				int32 LeftoverBurst = 0;
 				{
-					LeftoverBurst = BurstInfo.Count - FXConsoleVariables::MaxGPUParticlesSpawnedPerFrame;
-					BurstInfo.Count = FXConsoleVariables::MaxGPUParticlesSpawnedPerFrame;
+					float BurstDeltaTime = DeltaSeconds;
+					GetCurrentBurstRateOffset(BurstDeltaTime, BurstInfo.Count);
+
+					BurstInfo.Count += ForceBurstSpawnedParticles.Num();
+
+					if (BurstInfo.Count > FXConsoleVariables::MaxGPUParticlesSpawnedPerFrame)
+					{
+						LeftoverBurst = BurstInfo.Count - FXConsoleVariables::MaxGPUParticlesSpawnedPerFrame;
+						BurstInfo.Count = FXConsoleVariables::MaxGPUParticlesSpawnedPerFrame;
+					}
 				}
-			}
 
 
-			// Determine spawn count based on rate.
-			FSpawnInfo SpawnInfo = GetNumParticlesToSpawn(DeltaSeconds);
-			SpawnInfo.Count += ForceSpawnedParticles.Num();
+				// Determine spawn count based on rate.
+				FSpawnInfo SpawnInfo = GetNumParticlesToSpawn(DeltaSeconds);
+				SpawnInfo.Count += ForceSpawnedParticles.Num();
 
 
-			int32 FirstBurstParticleIndex = NewParticles.Num();
+				int32 FirstBurstParticleIndex = NewParticles.Num();
 
-			ReserveNewParticles(FirstBurstParticleIndex + BurstInfo.Count + SpawnInfo.Count);
+				ReserveNewParticles(FirstBurstParticleIndex + BurstInfo.Count + SpawnInfo.Count);
 
-			BurstInfo.Count = AllocateTilesForParticles(NewParticles, BurstInfo.Count, ActiveTileCount);
+				BurstInfo.Count = AllocateTilesForParticles(NewParticles, BurstInfo.Count, ActiveTileCount);
 
-			int32 FirstSpawnParticleIndex = NewParticles.Num();
-			SpawnInfo.Count = AllocateTilesForParticles(NewParticles, SpawnInfo.Count, ActiveTileCount);
-			SpawnFraction += LeftoverBurst;
+				int32 FirstSpawnParticleIndex = NewParticles.Num();
+				SpawnInfo.Count = AllocateTilesForParticles(NewParticles, SpawnInfo.Count, ActiveTileCount);
+				SpawnFraction += LeftoverBurst;
 
-			if (BurstInfo.Count > 0)
-			{
-				// Spawn burst particles.
-				BuildNewParticles(NewParticles.GetData() + FirstBurstParticleIndex, BurstInfo, ForceBurstSpawnedParticles);
-			}
+				if (BurstInfo.Count > 0)
+				{
+					// Spawn burst particles.
+					BuildNewParticles(NewParticles.GetData() + FirstBurstParticleIndex, BurstInfo, ForceBurstSpawnedParticles);
+				}
 
-			if (SpawnInfo.Count > 0)
-			{
-				// Spawn normal particles.
-				BuildNewParticles(NewParticles.GetData() + FirstSpawnParticleIndex, SpawnInfo, ForceSpawnedParticles);
-			}
+				if (SpawnInfo.Count > 0)
+				{
+					// Spawn normal particles.
+					BuildNewParticles(NewParticles.GetData() + FirstSpawnParticleIndex, SpawnInfo, ForceSpawnedParticles);
+				}
 
-			FreeNewParticleArray(ForceSpawnedParticles);
-			FreeNewParticleArray(ForceBurstSpawnedParticles);
+				FreeNewParticleArray(ForceSpawnedParticles);
+				FreeNewParticleArray(ForceBurstSpawnedParticles);
 
-			int32 NewParticleCount = BurstInfo.Count + SpawnInfo.Count;
-			INC_DWORD_STAT_BY(STAT_GPUSpritesSpawned, NewParticleCount);
-#if STATS
-			if (NewParticleCount > FXConsoleVariables::GPUSpawnWarningThreshold)
-			{
-				UE_LOG(LogParticles,Warning,TEXT("Spawning %d GPU particles in one frame[%d]: %s/%s"),
-					NewParticleCount,
-					GFrameNumber,
-					*SpriteTemplate->GetOuter()->GetName(),
-					*SpriteTemplate->EmitterName.ToString()
-					);
+				int32 NewParticleCount = BurstInfo.Count + SpawnInfo.Count;
+				INC_DWORD_STAT_BY(STAT_GPUSpritesSpawned, NewParticleCount);
+	#if STATS
+				if (NewParticleCount > FXConsoleVariables::GPUSpawnWarningThreshold)
+				{
+					UE_LOG(LogParticles,Warning,TEXT("Spawning %d GPU particles in one frame[%d]: %s/%s"),
+						NewParticleCount,
+						GFrameNumber,
+						*SpriteTemplate->GetOuter()->GetName(),
+						*SpriteTemplate->EmitterName.ToString()
+						);
 
-			}
-#endif
+				}
+	#endif
 
-			if (Component && Component->bWarmingUp)
+				if (Component && Component->bWarmingUp)
+				{
+					SimulateWarmupParticles(
+						NewParticles.GetData() + (NewParticles.Num() - NewParticleCount),
+						NewParticleCount,
+						Component->WarmupTime - SecondsSinceCreation );
+				}
+			}
+			else if (bFakeBurstsWhenSpawningSupressed)
 			{
-				SimulateWarmupParticles(
-					NewParticles.GetData() + (NewParticles.Num() - NewParticleCount),
-					NewParticleCount,
-					Component->WarmupTime - SecondsSinceCreation );
+				FakeBursts();
 			}
-		}
 
-		// Free any tiles that we no longer need.
-		FreeInactiveTiles();
+			// Free any tiles that we no longer need.
+			FreeInactiveTiles();
 
-		// Update current material.
-		if (EmitterInfo.RequiredModule->Material)
-		{
-			CurrentMaterial = EmitterInfo.RequiredModule->Material;
-		}
+			// Update current material.
+			if (EmitterInfo.RequiredModule->Material)
+			{
+				CurrentMaterial = EmitterInfo.RequiredModule->Material;
+			}
 
-		// Update the local vector field.
-		TickLocalVectorField(DeltaSeconds);
+			// Update the local vector field.
+			TickLocalVectorField(DeltaSeconds);
 
-		// Look up the strength of the point attractor.
-		EmitterInfo.PointAttractorStrength.GetValue(EmitterTime, &PointAttractorStrength);
+			// Look up the strength of the point attractor.
+			EmitterInfo.PointAttractorStrength.GetValue(EmitterTime, &PointAttractorStrength);
 
-		// Store the amount of time by which the GPU needs to update the simulation.
-		PendingDeltaSeconds = DeltaSeconds;
+			// Store the amount of time by which the GPU needs to update the simulation.
+			PendingDeltaSeconds = DeltaSeconds;
 
-		// Store the number of active particles.
-		ActiveParticles = ActiveTileCount * GParticlesPerTile;
-		INC_DWORD_STAT_BY(STAT_GPUSpriteParticles, ActiveParticles);
+			// Store the number of active particles.
+			ActiveParticles = ActiveTileCount * GParticlesPerTile;
+			INC_DWORD_STAT_BY(STAT_GPUSpriteParticles, ActiveParticles);
 
-		// 'Reset' the emitter time so that the delay functions correctly
-		EmitterTime += EmitterDelay;
+			// 'Reset' the emitter time so that the delay functions correctly
+			EmitterTime += EmitterDelay;
 
-		// Update the bounding box.
-		UpdateBoundingBox(DeltaSeconds);
+			// Update the bounding box.
+			UpdateBoundingBox(DeltaSeconds);
 
-		// Final update for modules.
-		Tick_ModuleFinalUpdate(DeltaSeconds, LODLevel);
+			// Final update for modules.
+			Tick_ModuleFinalUpdate(DeltaSeconds, LODLevel);
 
-		// Queue an update to the GPU simulation if needed.
-		if (Simulation->bDirty_GameThread)
-		{
-			Simulation->InitResources(AllocatedTiles, &EmitterInfo.Resources->EmitterSimulationResources);
-		}
+			// Queue an update to the GPU simulation if needed.
+			if (Simulation->bDirty_GameThread)
+			{
+				Simulation->InitResources(AllocatedTiles, &EmitterInfo.Resources->EmitterSimulationResources);
+			}
 
+			CheckEmitterFinished();
+		}
+		else
+		{
+			// 'Reset' the emitter time so that the delay functions correctly
+			EmitterTime += EmitterDelay;
 
-		CheckEmitterFinished();
+			FakeBursts();
+		}
 
 		check(AllocatedTiles.Num() == TileTimeOfDeath.Num());
 	}
@@ -4558,7 +4576,8 @@ void FFXSystem::SimulateGPUParticles(
 
 		FParticleSimulationGPU* Simulation = *It;
 		if (Simulation->SimulationPhase == Phase
-			&& Simulation->TileVertexBuffer.AlignedTileCount > 0)
+			&& Simulation->TileVertexBuffer.AlignedTileCount > 0
+			&& Simulation->bEnabled)
 		{
 			FSimulationCommandGPU* SimulationCommand = new(SimulationCommands) FSimulationCommandGPU(
 				Simulation->TileVertexBuffer.GetShaderParam(),
diff --git a/Engine/Source/Runtime/Engine/Private/Particles/ParticleSystemRender.cpp b/Engine/Source/Runtime/Engine/Private/Particles/ParticleSystemRender.cpp
--- a/Engine/Source/Runtime/Engine/Private/Particles/ParticleSystemRender.cpp
+++ b/Engine/Source/Runtime/Engine/Private/Particles/ParticleSystemRender.cpp
@@ -28,6 +28,8 @@ DECLARE_CYCLE_STAT(TEXT("DynamicSpriteEmitterData PerParticleWorkOrTasks"), STAT
 DECLARE_CYCLE_STAT(TEXT("DynamicSpriteEmitterData GetDynamicMeshElementsEmitter Task"), STAT_FDynamicSpriteEmitterData_GetDynamicMeshElementsEmitter_Task, STATGROUP_Particles);
 
 
+#include "InGamePerformanceTracker.h"
+
 /** 
  * Whether to track particle rendering stats.  
  * Enable with the TRACKPARTICLERENDERINGSTATS command. 
@@ -6431,6 +6433,7 @@ FParticleSystemSceneProxy::FParticleSystemSceneProxy(const UParticleSystemCompon
 	: FPrimitiveSceneProxy(Component, Component->Template ? Component->Template->GetFName() : NAME_None)
 	, Owner(Component->GetOwner())
 	, bCastShadow(Component->CastShadow)
+	, bManagingSignificance(Component->ShouldManageSignificance())
 	, MaterialRelevance(
 		((Component->GetCurrentLODIndex() >= 0) && (Component->GetCurrentLODIndex() < Component->CachedViewRelevanceFlags.Num())) ?
 			Component->CachedViewRelevanceFlags[Component->GetCurrentLODIndex()] :
@@ -6487,6 +6490,8 @@ FMeshBatch* FParticleSystemSceneProxy::GetPooledMeshBatch()
 
 void FParticleSystemSceneProxy::GetDynamicMeshElements(const TArray<const FSceneView*>& Views, const FSceneViewFamily& ViewFamily, uint32 VisibilityMap, FMeshElementCollector& Collector) const
 {
+	FInGameScopedCycleCounter InGameCycleCounter(GetScene().GetWorld(), EInGamePerfTrackers::VFXSignificance, EInGamePerfTrackerThreads::RenderThread, bManagingSignificance);
+
 	SCOPE_CYCLE_COUNTER(STAT_FParticleSystemSceneProxy_GetMeshElements);
 
 	if ((GIsEditor == true) || (GbEnableGameThreadLODCalculation == false))
@@ -6614,6 +6619,13 @@ void FParticleSystemSceneProxy::UpdateData(FParticleDynamicData* NewDynamicData)
 		{
 			SCOPE_CYCLE_COUNTER(STAT_ParticleUpdateRTTime);
 			STAT(FScopeCycleCounter Context(Proxy->GetStatId());)
+			if (NewDynamicData)
+			{
+				for (int32 Index = 0; Index < NewDynamicData->DynamicEmitterDataArray.Num(); Index++)
+				{
+					Proxy->QueueVertexFactoryCreation(NewDynamicData->DynamicEmitterDataArray[Index]);
+				}
+			}
 			Proxy->UpdateData_RenderThread(NewDynamicData);
 		}
 		);
@@ -6621,6 +6633,7 @@ void FParticleSystemSceneProxy::UpdateData(FParticleDynamicData* NewDynamicData)
 
 void FParticleSystemSceneProxy::UpdateData_RenderThread(FParticleDynamicData* NewDynamicData)
 {
+	FInGameScopedCycleCounter InGameCycleCounter(GetScene().GetWorld(), EInGamePerfTrackers::VFXSignificance, EInGamePerfTrackerThreads::RenderThread, bManagingSignificance);
 
 	ReleaseRenderThreadResourcesForEmitterData();
 	if (DynamicData != NewDynamicData)
@@ -6808,6 +6821,7 @@ void FParticleSystemSceneProxy::UpdateWorldSpacePrimitiveUniformBuffer() const
 
 void FParticleSystemSceneProxy::GatherSimpleLights(const FSceneViewFamily& ViewFamily, FSimpleLightArray& OutParticleLights) const
 {
+	FInGameScopedCycleCounter InGameCycleCounter(GetScene().GetWorld(), EInGamePerfTrackers::VFXSignificance, EInGamePerfTrackerThreads::RenderThread, bManagingSignificance);
 	if (DynamicData != NULL)
 	{
 		FScopeCycleCounter Context(GetStatId());
@@ -6844,11 +6858,13 @@ FParticleSystemOcclusionSceneProxy::~FParticleSystemOcclusionSceneProxy()
 
 FPrimitiveSceneProxy* UParticleSystemComponent::CreateSceneProxy()
 {
+	bool bManagingSignificance = ShouldManageSignificance();
+	FInGameScopedCycleCounter InGameCycleCounter(GetWorld(), EInGamePerfTrackers::VFXSignificance, EInGamePerfTrackerThreads::GameThread, bManagingSignificance);
 	FParticleSystemSceneProxy* NewProxy = NULL;
 
 	//@fixme EmitterInstances.Num() check should be here to avoid proxies for dead emitters but there are some edge cases where it happens for emitters that have just activated...
 	//@fixme Get non-instanced path working in ES2!
-	if ((bIsActive == true)/** && (EmitterInstances.Num() > 0)*/ && Template)
+	if ((bIsActive == true)/** && (EmitterInstances.Num() > 0)*/ && Template && bIsSignificant)
 	{
 		UE_LOG(LogParticles,Verbose,
 			TEXT("CreateSceneProxy @ %fs %s bIsActive=%d"), GetWorld()->TimeSeconds,
diff --git a/Engine/Source/Runtime/Engine/Private/Particles/ParticleTrail2EmitterInstance.cpp b/Engine/Source/Runtime/Engine/Private/Particles/ParticleTrail2EmitterInstance.cpp
--- a/Engine/Source/Runtime/Engine/Private/Particles/ParticleTrail2EmitterInstance.cpp
+++ b/Engine/Source/Runtime/Engine/Private/Particles/ParticleTrail2EmitterInstance.cpp
@@ -92,45 +92,54 @@ void FParticleTrailsEmitterInstance_Base::Tick(float DeltaTime, bool bSuppressSp
 		// Handle EmitterTime setup, looping, etc.
 		float EmitterDelay = Tick_EmitterTimeSetup(DeltaTime, LODLevel);
 
-		// Update the source data (position, etc.)
-		UpdateSourceData(DeltaTime, bFirstTime);
+		if (bEnabled)
+		{
+			// Update the source data (position, etc.)
+			UpdateSourceData(DeltaTime, bFirstTime);
 
-		// Kill off any dead particles
-		KillParticles();
+			// Kill off any dead particles
+			KillParticles();
 
-		// Spawn Particles...
-		SpawnFraction = Tick_SpawnParticles(DeltaTime, LODLevel, bSuppressSpawning, bFirstTime);
+			// Spawn Particles...
+			SpawnFraction = Tick_SpawnParticles(DeltaTime, LODLevel, bSuppressSpawning, bFirstTime);
 
-		// Reset particle parameters.
-		ResetParticleParameters(DeltaTime);
+			// Reset particle parameters.
+			ResetParticleParameters(DeltaTime);
 
-		// Update existing particles (might respawn dying ones).
-		Tick_ModuleUpdate(DeltaTime, LODLevel);
+			// Update existing particles (might respawn dying ones).
+			Tick_ModuleUpdate(DeltaTime, LODLevel);
 
-		// Module post update 
-		Tick_ModulePostUpdate(DeltaTime, LODLevel);
+			// Module post update 
+			Tick_ModulePostUpdate(DeltaTime, LODLevel);
 
-		// Update the orbit data...
-	// 	UpdateOrbitData(DeltaTime);
+			// Update the orbit data...
+			// 	UpdateOrbitData(DeltaTime);
 
-		// Calculate bounding box and simulate velocity.
-		UpdateBoundingBox(DeltaTime);
+			// Calculate bounding box and simulate velocity.
+			UpdateBoundingBox(DeltaTime);
 
-		// Perform any final updates...
-		Tick_ModuleFinalUpdate(DeltaTime, LODLevel);
+			// Perform any final updates...
+			Tick_ModuleFinalUpdate(DeltaTime, LODLevel);
 
-		// Recalculate tangents, if enabled
-		Tick_RecalculateTangents(DeltaTime, LODLevel);
+			// Recalculate tangents, if enabled
+			Tick_RecalculateTangents(DeltaTime, LODLevel);
 
-		CurrentMaterial = LODLevel->RequiredModule->Material;
+			CurrentMaterial = LODLevel->RequiredModule->Material;
 
-		// Invalidate the contents of the vertex/index buffer.
-		IsRenderDataDirty = 1;
+			// Invalidate the contents of the vertex/index buffer.
+			IsRenderDataDirty = 1;
 
-		// 'Reset' the emitter time so that the delay functions correctly
-		EmitterTime += CurrentDelay;
-		RunningTime += DeltaTime;
-		LastTickTime = Component->GetWorld() ? Component->GetWorld()->GetTimeSeconds() : 0.0f;
+			// 'Reset' the emitter time so that the delay functions correctly
+			EmitterTime += CurrentDelay;
+			RunningTime += DeltaTime;
+			LastTickTime = Component->GetWorld() ? Component->GetWorld()->GetTimeSeconds() : 0.0f;
+		}
+		else
+		{
+			// 'Reset' the emitter time so that the delay functions correctly
+			EmitterTime += CurrentDelay;
+			FakeBursts();
+		}
 		
 		// Reset particles position offset
 		PositionOffsetThisTick = FVector::ZeroVector;
@@ -2588,7 +2597,7 @@ FDynamicEmitterDataBase* FParticleRibbonEmitterInstance::GetDynamicData(bool bSe
 	QUICK_SCOPE_CYCLE_COUNTER(STAT_ParticleRibbonEmitterInstance_GetDynamicData);
 
 	UParticleLODLevel* LODLevel = SpriteTemplate->GetLODLevel(0);
-	if (IsDynamicDataRequired(LODLevel) == false)
+	if (IsDynamicDataRequired(LODLevel) == false || !bEnabled)
 	{
 		return NULL;
 	}
@@ -2643,7 +2652,7 @@ bool FParticleRibbonEmitterInstance::UpdateDynamicData(FDynamicEmitterDataBase*
 {
 	QUICK_SCOPE_CYCLE_COUNTER(STAT_ParticleRibbonEmitterInstance_UpdateDynamicData);
 
-	if (ActiveParticles <= 0)
+	if (ActiveParticles <= 0 || !bEnabled)
 	{
 		return false;
 	}
@@ -2681,7 +2690,7 @@ bool FParticleRibbonEmitterInstance::UpdateDynamicData(FDynamicEmitterDataBase*
  */
 FDynamicEmitterReplayDataBase* FParticleRibbonEmitterInstance::GetReplayData()
 {
-	if (ActiveParticles <= 0)
+	if (ActiveParticles <= 0 || !bEnabled)
 	{
 		return NULL;
 	}
@@ -2756,7 +2765,7 @@ bool FParticleRibbonEmitterInstance::FillReplayData(FDynamicEmitterReplayDataBas
 {
 	QUICK_SCOPE_CYCLE_COUNTER(STAT_ParticleRibbonEmitterInstance_FillReplayData);
 
-	if (ActiveParticles <= 0)
+	if (ActiveParticles <= 0 || !bEnabled)
 	{
 		return false;
 	}
@@ -3940,7 +3949,7 @@ FDynamicEmitterDataBase* FParticleAnimTrailEmitterInstance::GetDynamicData(bool
 	QUICK_SCOPE_CYCLE_COUNTER(STAT_ParticleAnimTrailEmitterInstance_GetDynamicData);
 
 	UParticleLODLevel* LODLevel = SpriteTemplate->GetLODLevel(0);
-	if (IsDynamicDataRequired(LODLevel) == false)
+	if (IsDynamicDataRequired(LODLevel) == false || !bEnabled)
 	{
 		return NULL;
 	}
@@ -4004,7 +4013,7 @@ bool FParticleAnimTrailEmitterInstance::UpdateDynamicData(FDynamicEmitterDataBas
 {
 	QUICK_SCOPE_CYCLE_COUNTER(STAT_ParticleAnimTrailEmitterInstance_UpdateDynamicData);
 
-	if (ActiveParticles <= 0)
+	if (ActiveParticles <= 0 || !bEnabled)
 	{
 		return false;
 	}
@@ -4047,7 +4056,7 @@ bool FParticleAnimTrailEmitterInstance::UpdateDynamicData(FDynamicEmitterDataBas
  */
 FDynamicEmitterReplayDataBase* FParticleAnimTrailEmitterInstance::GetReplayData()
 {
-	if (ActiveParticles <= 0)
+	if (ActiveParticles <= 0 || !bEnabled)
 	{
 		return NULL;
 	}
@@ -4177,7 +4186,7 @@ bool FParticleAnimTrailEmitterInstance::FillReplayData( FDynamicEmitterReplayDat
 {
 	QUICK_SCOPE_CYCLE_COUNTER(STAT_ParticleAnimTrailEmitterInstance_FillReplayData);
 
-	if (ActiveParticles <= 0)
+	if (ActiveParticles <= 0 || !bEnabled)
 	{
 		return false;
 	}
diff --git a/Engine/Source/Runtime/Engine/Private/RepLayout.cpp b/Engine/Source/Runtime/Engine/Private/RepLayout.cpp
--- a/Engine/Source/Runtime/Engine/Private/RepLayout.cpp
+++ b/Engine/Source/Runtime/Engine/Private/RepLayout.cpp
@@ -2560,6 +2560,7 @@ void FRepLayout::RebuildConditionalProperties( FRepState * RESTRICT	RepState, co
 
 	ConditionMap[COND_InitialOrOwner]		= bIsInitial || bIsOwner;
 	ConditionMap[COND_ReplayOrOwner]		= bIsReplay || bIsOwner;
+	ConditionMap[COND_ReplayOnly]			= bIsReplay;
 
 	ConditionMap[COND_Custom]				= true;
 
diff --git a/Engine/Source/Runtime/Engine/Private/SoundNode.cpp b/Engine/Source/Runtime/Engine/Private/SoundNode.cpp
--- a/Engine/Source/Runtime/Engine/Private/SoundNode.cpp
+++ b/Engine/Source/Runtime/Engine/Private/SoundNode.cpp
@@ -8,6 +8,7 @@
 #if WITH_EDITOR
 #include "UnrealEd.h"
 #endif
+#include "AssetData.h"
 
 /*-----------------------------------------------------------------------------
 	USoundNode implementation.
@@ -183,4 +184,5 @@ void USoundNode::PlaceNode( int32 NodeColumn, int32 NodeRow, int32 RowCount )
 	GraphNode->NodePosX = (-150 * NodeColumn) - 100;
 	GraphNode->NodePosY = (100 * NodeRow) - (50 * RowCount);
 }
+
 #endif //WITH_EDITOR
diff --git a/Engine/Source/Runtime/Engine/Private/SoundNodeDialoguePlayer.cpp b/Engine/Source/Runtime/Engine/Private/SoundNodeDialoguePlayer.cpp
--- a/Engine/Source/Runtime/Engine/Private/SoundNodeDialoguePlayer.cpp
+++ b/Engine/Source/Runtime/Engine/Private/SoundNodeDialoguePlayer.cpp
@@ -5,6 +5,9 @@
 #include "Sound/DialogueWave.h"
 #include "Sound/SoundBase.h"
 #include "Sound/SoundNodeDialoguePlayer.h"
+#include "AssetData.h"
+
+#define LOCTEXT_NAMESPACE "SoundNodeDialoguePlayer"
 
 USoundNodeDialoguePlayer::USoundNodeDialoguePlayer(const FObjectInitializer& ObjectInitializer)
 	: Super(ObjectInitializer)
@@ -13,7 +16,7 @@ USoundNodeDialoguePlayer::USoundNodeDialoguePlayer(const FObjectInitializer& Obj
 
 void USoundNodeDialoguePlayer::ParseNodes( FAudioDevice* AudioDevice, const UPTRINT NodeWaveInstanceHash, FActiveSound& ActiveSound, const FSoundParseParameters& ParseParams, TArray<FWaveInstance*>& WaveInstances )
 {
-	USoundBase* SoundBase = DialogueWaveParameter.DialogueWave ? DialogueWaveParameter.DialogueWave->GetWaveFromContext(DialogueWaveParameter.Context) : NULL;
+	USoundBase* SoundBase = GetDialogueWave() ? GetDialogueWave()->GetWaveFromContext(DialogueWaveParameter.Context) : NULL;
 	if (SoundBase)
 	{
 		if (bLooping)
@@ -31,7 +34,7 @@ void USoundNodeDialoguePlayer::ParseNodes( FAudioDevice* AudioDevice, const UPTR
 
 float USoundNodeDialoguePlayer::GetDuration()
 {
-	USoundBase* SoundBase = DialogueWaveParameter.DialogueWave ? DialogueWaveParameter.DialogueWave->GetWaveFromContext(DialogueWaveParameter.Context) : NULL;
+	USoundBase* SoundBase = GetDialogueWave() ? GetDialogueWave()->GetWaveFromContext(DialogueWaveParameter.Context) : NULL;
 	float Duration = 0.f;
 	if (SoundBase)
 	{
@@ -52,3 +55,46 @@ int32 USoundNodeDialoguePlayer::GetMaxChildNodes() const
 {
 	return 0;
 }
+
+#if WITH_EDITOR
+FText USoundNodeDialoguePlayer::GetTitle() const
+{
+	FText DialogueWaveName;
+	if (GetDialogueWave())
+	{
+		DialogueWaveName = FText::FromString(GetDialogueWave()->GetFName().ToString());
+	}
+	else
+	{
+		DialogueWaveName = LOCTEXT("NoDialogueWave", "NONE");
+	}
+
+	FText Title;
+
+	FFormatNamedArguments Arguments;
+	Arguments.Add(TEXT("Description"), Super::GetTitle());
+	Arguments.Add(TEXT("DialogueWaveName"), DialogueWaveName);
+	if (bLooping)
+	{
+		Title = FText::Format(LOCTEXT("LoopingDialogueWaveDescription", "Looping {Description} : {DialogueWaveName}"), Arguments);
+	}
+	else
+	{
+		Title = FText::Format(LOCTEXT("NonLoopingDialogueWaveDescription", "{Description} : {DialogueWaveName}"), Arguments);
+	}
+
+	return Title;
+}
+#endif
+
+void USoundNodeDialoguePlayer::SetDialogueWave(UDialogueWave* Value)
+{
+	DialogueWaveParameter.DialogueWave = Value;
+}
+
+UDialogueWave* USoundNodeDialoguePlayer::GetDialogueWave() const
+{
+	return DialogueWaveParameter.DialogueWave;
+}
+
+#undef LOCTEXT_NAMESPACE
diff --git a/Engine/Source/Runtime/Engine/Private/TextureDerivedData.cpp b/Engine/Source/Runtime/Engine/Private/TextureDerivedData.cpp
--- a/Engine/Source/Runtime/Engine/Private/TextureDerivedData.cpp
+++ b/Engine/Source/Runtime/Engine/Private/TextureDerivedData.cpp
@@ -38,77 +38,6 @@ enum
 
 #define TEXTURE_DERIVEDDATA_VER		TEXT("814DCC3DC72143F49509781513CB9855")
 
-
-// output texture building stats to Saved\Stats\Stats.csv
-#define BUILD_TEXTURE_STATS 1 && ENABLE_DDC_STATS
-
-#if BUILD_TEXTURE_STATS
-#include "DDCStatsHelper.h"
-#endif
-
-/*------------------------------------------------------------------------------
-	Timing of derived data operations.
-------------------------------------------------------------------------------*/
-
-namespace TextureDerivedDataTimings
-{
-	enum ETimingId
-	{
-		GetMipDataTime,
-		AsyncBlockTime,
-		SyncBlockTime,
-		BuildTextureTime,
-		SerializeCookedTime,
-		SyncDoesCachedDataExist,
-		NumTimings
-	};
-
-	static FThreadSafeCounter Timings[NumTimings] = {0};
-
-	static const TCHAR* TimingStrings[NumTimings] =
-	{
-		TEXT("Get Mip Data"),
-		TEXT("Asynchronous Block"),
-		TEXT("Synchronous Loads"),
-		TEXT("Build Textures"),
-		TEXT("Serialize Cooked"),
-		TEXT("Sync Does Cached Data Exist")
-	};
-
-	void PrintTimings()
-	{
-		UE_LOG(LogTexture,Log,TEXT("--- Texture Derived Data Timings ---"));
-		for (int32 TimingIndex = 0; TimingIndex < NumTimings; ++TimingIndex)
-		{
-			UE_LOG(LogTexture,Display,TEXT("%s: %fs"), TimingStrings[TimingIndex], FPlatformTime::ToSeconds(Timings[TimingIndex].GetValue()) );
-		}
-	}
-
-	FAutoConsoleCommand DumpTimingsCommand(
-		TEXT("Tex.DerivedDataTimings"),
-		TEXT("Print timings related to texture derived data."),
-		FConsoleCommandDelegate::CreateStatic(PrintTimings)
-		);
-
-	struct FScopedMeasurement
-	{
-		ETimingId TimingId;
-		uint32 StartCycles;
-
-		explicit FScopedMeasurement(ETimingId InTimingId)
-			: TimingId(InTimingId)
-		{
-			StartCycles = FPlatformTime::Cycles();
-		}
-
-		~FScopedMeasurement()
-		{
-			uint32 TimeInCycles = FPlatformTime::Cycles() - StartCycles;
-			Timings[TimingId].Add(TimeInCycles);
-		}
-	};
-}
-
 /*------------------------------------------------------------------------------
 	Derived data key generation.
 ------------------------------------------------------------------------------*/
@@ -774,18 +703,8 @@ class FTextureCacheDerivedDataWorker : public FNonAbandonableTask
 	/** Build the texture. This function is safe to call from any thread. */
 	void BuildTexture()
 	{
-#if BUILD_TEXTURE_STATS
-		FString DerivedDataKey;
-		GetTextureDerivedDataKeyFromSuffix(KeySuffix, DerivedDataKey);
-		static const FName NAME_BuildTexture(TEXT("BuildTexture"));
-		FDDCScopeStatHelper ScopeStat(*DerivedDataKey, NAME_BuildTexture);
-#endif
-
-
 		if (SourceMips.Num())
 		{
-			TextureDerivedDataTimings::FScopedMeasurement Timer(TextureDerivedDataTimings::BuildTextureTime);
-
 			FFormatNamedArguments Args;
 			Args.Add( TEXT("TextureName"), FText::FromString( Texture.GetName() ) );
 			Args.Add( TEXT("TextureFormatName"), FText::FromString( BuildSettings.TextureFormatName.GetPlainNameString() ) );
@@ -1012,7 +931,6 @@ void FTexturePlatformData::Cache(
 	}
 	else
 	{
-		TextureDerivedDataTimings::FScopedMeasurement Timer(TextureDerivedDataTimings::SyncBlockTime);
 		FTextureCacheDerivedDataWorker Worker(Compressor, this, &InTexture, InSettings, Flags);
 		Worker.DoWork();
 		Worker.Finalize();
@@ -1024,7 +942,6 @@ void FTexturePlatformData::FinishCache()
 	if (AsyncTask)
 	{
 		{
-			TextureDerivedDataTimings::FScopedMeasurement Timer(TextureDerivedDataTimings::AsyncBlockTime);
 			AsyncTask->EnsureCompletion();
 		}
 		FTextureCacheDerivedDataWorker& Worker = AsyncTask->GetTask();
@@ -1398,9 +1315,6 @@ void FAsyncStreamDerivedMipWorker::DoWork()
 
 void UTexture2D::GetMipData(int32 FirstMipToLoad, void** OutMipData)
 {
-#if WITH_EDITOR
-	TextureDerivedDataTimings::FScopedMeasurement MipDataTimer(TextureDerivedDataTimings::GetMipDataTime);
-#endif // #if WITH_EDITOR
 	if (PlatformData->TryLoadMips(FirstMipToLoad, OutMipData) == false)
 	{
 		// Unable to load mips from the cache. Rebuild the texture and try again.
@@ -1538,7 +1452,6 @@ void UTexture::BeginCacheForCookedPlatformData( const ITargetPlatform *TargetPla
 				// if it doesn't then allow async builds
 				
 				{
-					TextureDerivedDataTimings::FScopedMeasurement Timer(TextureDerivedDataTimings::SyncDoesCachedDataExist);
 					if ( GetDerivedDataCacheRef().CachedDataProbablyExists( *DerivedDataKey ) == false )
 					{
 						CurrentCacheFlags |= ETextureCacheFlags::AllowAsyncBuild;
@@ -1849,7 +1762,6 @@ void UTexture::SerializeCookedPlatformData(FArchive& Ar)
 
 
 #if WITH_EDITOR
-	TextureDerivedDataTimings::FScopedMeasurement Timer(TextureDerivedDataTimings::SerializeCookedTime);
 	if (Ar.IsCooking() && Ar.IsPersistent())
 	{
 		if (!Ar.CookingTarget()->IsServerOnly())
diff --git a/Engine/Source/Runtime/Engine/Private/UnrealEngine.cpp b/Engine/Source/Runtime/Engine/Private/UnrealEngine.cpp
--- a/Engine/Source/Runtime/Engine/Private/UnrealEngine.cpp
+++ b/Engine/Source/Runtime/Engine/Private/UnrealEngine.cpp
@@ -2345,11 +2345,26 @@ struct FSubItem
 
 MSVC_PRAGMA(warning(push))
 MSVC_PRAGMA(warning(disable : 4717))
-static void InfiniteRecursionFunction(bool B)
+
+#if defined (__clang__) 
+	#pragma clang diagnostic push
+	#pragma clang diagnostic ignored "-Winfinite-recursion"
+#endif
+
+
+// clang can optimize this out (http://stackoverflow.com/questions/18478078/clang-infinite-tail-recursion-optimization), make the function do some useful work
+
+volatile int GInfiniteRecursionCount = 0;
+static int InfiniteRecursionFunction(int B)
 {
-	if(B)
-		InfiniteRecursionFunction(B);
+	GInfiniteRecursionCount += InfiniteRecursionFunction(B + 1);
+	return GInfiniteRecursionCount;
 }
+
+#if defined (__clang__) 
+	#pragma clang diagnostic pop
+#endif
+
 MSVC_PRAGMA(warning(pop))
 
 /** DEBUG used for exe "DEBUG BUFFEROVERFLOW" */
@@ -4739,14 +4754,14 @@ bool UEngine::HandleDebugCommand( const TCHAR* Cmd, FOutputDevice& Ar )
 	{
 		UE_LOG(LogEngine, Log, TEXT("This is going to be a really long log message to test the code to resize the buffer used to log with. %02048s"), TEXT("HAHA, this isn't really a long string, but it sure has a lot of zeros!"));
 	}
-#if 0
 	else if( FParse::Command( &Cmd, TEXT("RECURSE") ) )
 	{
 		Ar.Logf( TEXT("Recursing") );
+		GLog->Flush();
 		InfiniteRecursionFunction(1);
+		Ar.Logf(TEXT("You will never see this log line."));
 		return true;
 	}
-#endif
 	else if( FParse::Command( &Cmd, TEXT("EATMEM") ) )
 	{
 		Ar.Log( TEXT("Eating up all available memory") );
@@ -9832,7 +9847,7 @@ bool UEngine::LoadMap( FWorldContext& WorldContext, FURL URL, class UPendingNetG
 
 	double StopTime = FPlatformTime::Seconds();
 
-	UE_LOG(LogEngine, Log, TEXT("Took %f seconds to LoadMap(%s)"), StopTime - StartTime, *URL.Map);
+	UE_LOG(LogLoad, Log, TEXT("Took %f seconds to LoadMap(%s)"), StopTime - StartTime, *URL.Map);
 
 	// Successfully started local level.
 	return true;
diff --git a/Engine/Source/Runtime/Engine/Private/World.cpp b/Engine/Source/Runtime/Engine/Private/World.cpp
--- a/Engine/Source/Runtime/Engine/Private/World.cpp
+++ b/Engine/Source/Runtime/Engine/Private/World.cpp
@@ -85,6 +85,7 @@
 #include "Engine/CoreSettings.h"
 #include "PerfCountersHelpers.h"
 #include "NetworkReplayStreaming.h"
+#include "InGamePerformanceTracker.h"
 
 DEFINE_LOG_CATEGORY_STATIC(LogWorld, Log, All);
 DEFINE_LOG_CATEGORY(LogSpawn);
@@ -182,6 +183,8 @@ UWorld::UWorld( const FObjectInitializer& ObjectInitializer )
 #endif // WITH_EDITOR
 
 	FWorldDelegates::OnPostWorldCreation.Broadcast(this);
+
+	PerfTrackers = new FWorldInGamePerformanceTrackers();
 }
 
 UWorld::~UWorld()
@@ -190,6 +193,11 @@ UWorld::~UWorld()
 	{
 		FPlatformProcess::Sleep(0.0f);
 	}	
+
+	if (PerfTrackers)
+	{
+		delete PerfTrackers;
+	}
 }
 
 void UWorld::Serialize( FArchive& Ar )
diff --git a/Engine/Source/Runtime/Engine/Public/Animation/AnimInstanceProxy.h b/Engine/Source/Runtime/Engine/Public/Animation/AnimInstanceProxy.h
--- a/Engine/Source/Runtime/Engine/Public/Animation/AnimInstanceProxy.h
+++ b/Engine/Source/Runtime/Engine/Public/Animation/AnimInstanceProxy.h
@@ -85,6 +85,14 @@ public:
 	void TickSyncGroupWriteIndex()
 	{ 
 		SyncGroupWriteIndex = GetSyncGroupReadIndex();
+
+		// Reset 'write' SyncGroups for next frame.
+		TArray<FAnimGroupInstance>& SyncGroups = SyncGroupArrays[SyncGroupWriteIndex];
+		for (int32 GroupIndex = 0; GroupIndex < SyncGroups.Num(); ++GroupIndex)
+		{
+			FAnimGroupInstance& SyncGroup = SyncGroups[GroupIndex];
+			SyncGroup = FAnimGroupInstance();
+		}
 	}
 
 	/** Get the sync group we are currently reading from */
diff --git a/Engine/Source/Runtime/Engine/Public/InGamePerformanceTracker.h b/Engine/Source/Runtime/Engine/Public/InGamePerformanceTracker.h
new file mode 100644
--- /dev/null
+++ b/Engine/Source/Runtime/Engine/Public/InGamePerformanceTracker.h
@@ -0,0 +1,191 @@
+// Copyright 1998-2015 Epic Games, Inc. All Rights Reserved.
+
+#pragma once
+
+struct FInGameCycleHistory
+{
+	FInGameCycleHistory(int32 InNumSamples)
+	: FrameIdx(0)
+	, TotalCycles(0)
+	, ValidFrames(0)
+	, CurrFrameCycles(0)
+	{
+		FrameCycles.Empty(InNumSamples);
+		FrameCycles.SetNumZeroed(InNumSamples);
+	}
+
+	void Reset()
+	{
+		int32 Num = FrameCycles.Num();
+		FrameCycles.Empty(Num);
+		FrameCycles.SetNumZeroed(Num);
+		FrameIdx = 0;
+		TotalCycles = 0;
+		ValidFrames = 0;
+		CurrFrameCycles = 0;
+	}
+
+	//Cycles for each frame in history.
+	TArray<uint32> FrameCycles;
+	//Index of next frame in history to replace.
+	int32 FrameIdx;
+	//Cached running total for an overall average.
+	uint64 TotalCycles;
+	/** Number of current valid frames. A history is only valid when all it's frames are. */
+	int32 ValidFrames;
+	/** Current number of cycles for this frame. */
+	int32 CurrFrameCycles;
+
+	/** Adds a number of cycles for the current frame. Thread safe. */
+	FORCEINLINE void AddCycles(int32 NewCycles)
+	{
+		FPlatformAtomics::InterlockedAdd(&CurrFrameCycles, NewCycles);
+	}
+	
+	FORCEINLINE uint32 GetAverageCycles()const
+	{
+		return IsValid() ? TotalCycles / ValidFrames : 0;
+	}
+
+	/** A history only becomes valid when all it's frames are in use. e.g. First N frames of tracking are not valid. */
+	FORCEINLINE bool IsValid()const
+	{
+		return ValidFrames == FrameCycles.Num(); 
+	}
+
+	/** Places the current data on the history and advances to the next frame index. Not thread safe.*/
+	void NextFrame();
+};
+
+/**
+Helper class to track code timings.
+Similar to stats but much more minimal as they're intended to be used at InGame. Though it's still not free so use wisely.
+*/
+class FInGamePerformanceTracker
+{
+private:
+
+	/** An N frame history of cycle counts. */
+	FInGameCycleHistory History;
+
+	/** How many times we've entered a directly tracked section of code. Used to time the first and last entry/exit points. */
+	uint32 DirectSectionTime_EntryCount;
+	/** Initial cycle reading when timing a section of code directly. */
+	uint32 DirectSectionTime_BeginCycles;
+	
+public:
+
+	FInGamePerformanceTracker();
+	FInGamePerformanceTracker(uint32 FrameHistorySize);
+
+	void Tick();
+
+	/** Adds cycles collected from an external helper object. */
+	FORCEINLINE void AddCycles(uint32 Cycles)
+	{
+		return History.AddCycles(Cycles);
+	}
+
+	FORCEINLINE float GetAverageTimeSeconds()const
+	{
+		return FPlatformTime::ToSeconds(History.GetAverageCycles());
+	}
+	
+	/** Enters a timed section of code. Can be re-entrant and will take the first entry point time. Only callable from game thread. */
+	void EnterTimedSection();
+
+	/** Enters a timed section of code. Can be re-entrant and will take the final exit point time. Only callable from game thread. */
+	void ExitTimedSection();
+	
+	ENGINE_API static class IConsoleVariable* Enabled;
+	ENGINE_API static int32 CachedEnabled;
+	ENGINE_API static class IConsoleVariable* HistorySize;
+};
+
+class FInGameCycleCounter
+{
+private:
+	FInGamePerformanceTracker* Tracker;
+	uint32 BeginCycles;
+
+public:
+	FInGameCycleCounter(FInGamePerformanceTracker* InTracker)
+	: Tracker(InTracker)
+	, BeginCycles(0)
+	{
+	}
+
+	FORCEINLINE void Begin()
+	{
+		if (Tracker && FInGamePerformanceTracker::CachedEnabled)
+		{
+			if (IsInGameThread())
+			{
+				Tracker->EnterTimedSection();//On the GT we use the tracker directly so we can allow re-entrance.
+			}
+			else
+			{
+				//On other threads we don't allow (or check for) re-entrance and store our own cycle data.
+				BeginCycles = FPlatformTime::Cycles();
+			}
+		}
+	}
+	FORCEINLINE void End()
+	{
+		if (Tracker && FInGamePerformanceTracker::CachedEnabled)
+		{
+			if (IsInGameThread())
+			{
+				Tracker->ExitTimedSection();//On the GT we use the tracker directly so we can allow re-entrance.
+			}
+			else
+			{
+				//On other threads we don't allow (or check for) re-entrance and store our own cycle data.
+				uint32 Cycles = FPlatformTime::Cycles() - BeginCycles;
+				Tracker->AddCycles(Cycles);
+				BeginCycles = 0;
+			}
+		}
+	}
+};
+
+// In game performance trackers. 
+// Not the most extensible system so should likely be improved in the future.
+enum class EInGamePerfTrackers : uint8
+{
+	VFXSignificance,
+	//Others?
+	Num,
+};
+
+enum class EInGamePerfTrackerThreads : uint8
+{
+	GameThread,
+	RenderThread,
+	OtherThread,
+	Num,
+};
+
+class FInGameScopedCycleCounter : public FInGameCycleCounter
+{
+public:
+	FInGameScopedCycleCounter(class UWorld* InWorld, EInGamePerfTrackers Tracker, EInGamePerfTrackerThreads TrackerThread, bool bEnabled);
+	~FInGameScopedCycleCounter();
+};
+
+
+/** Collection of in game performance trackers for a world. */
+class FWorldInGamePerformanceTrackers
+{
+	FInGamePerformanceTracker InGamePerformanceTrackers[(int32)EInGamePerfTrackers::Num][(int32)EInGamePerfTrackerThreads::Num];
+
+public:
+
+	FWorldInGamePerformanceTrackers();
+	~FWorldInGamePerformanceTrackers();
+
+	FInGamePerformanceTracker& GetInGamePerformanceTracker(EInGamePerfTrackers Tracker, EInGamePerfTrackerThreads Thread)
+	{
+		return InGamePerformanceTrackers[(int32)Tracker][(int32)Thread];
+	}
+};
\ No newline at end of file
diff --git a/Engine/Source/Runtime/Engine/Public/ParticleEmitterInstances.h b/Engine/Source/Runtime/Engine/Public/ParticleEmitterInstances.h
--- a/Engine/Source/Runtime/Engine/Public/ParticleEmitterInstances.h
+++ b/Engine/Source/Runtime/Engine/Public/ParticleEmitterInstances.h
@@ -233,6 +233,8 @@ public:
 	uint32 bIsBeam:1;
 	/** Whether axis lock is enabled, cached here to avoid finding it from the module each frame */
 	uint32 bAxisLockEnabled : 1;
+	/** When true and spawning is supressed, the bursts will be faked so that when spawning is enabled again, the bursts don't fire late. */
+	uint32 bFakeBurstsWhenSpawningSupressed : 1;
 	/** Axis lock flags, cached here to avoid finding it from the module each frame */
 	TEnumAsByte<EParticleAxisLock> LockAxisFlags;
 	/** The sort mode to use for this emitter as specified by artist.	*/
@@ -337,6 +339,9 @@ public:
 	virtual void Tick(float DeltaTime, bool bSuppressSpawning);
 	void CheckEmitterFinished();
 
+	/** Advances the bursts as though they were fired with out actually firing them. */
+	void FakeBursts();
+
 	/**
 	 *	Tick sub-function that handles EmitterTime setup, looping, etc.
 	 *
@@ -479,6 +484,11 @@ public:
 		bHaltSpawning = bInHaltSpawning;
 	}
 
+	FORCEINLINE void SetFakeBurstWhenSpawningSupressed(bool bInFakeBurstsWhenSpawningSupressed)
+	{
+		bFakeBurstsWhenSpawningSupressed = bInFakeBurstsWhenSpawningSupressed;
+	}
+
 	/** Get the offset of the orbit payload. */
 	int32 GetOrbitPayloadOffset();
 
diff --git a/Engine/Source/Runtime/Engine/Public/ParticleHelper.h b/Engine/Source/Runtime/Engine/Public/ParticleHelper.h
--- a/Engine/Source/Runtime/Engine/Public/ParticleHelper.h
+++ b/Engine/Source/Runtime/Engine/Public/ParticleHelper.h
@@ -2587,6 +2587,7 @@ protected:
 #endif
 
 	uint32 bCastShadow : 1;
+	uint32 bManagingSignificance : 1;
 	
 	FMaterialRelevance MaterialRelevance;
 
@@ -2773,4 +2774,21 @@ enum class EParticleSignificanceLevel : uint8
 	Critical,
 
 	Num UMETA(Hidden),
-};
\ No newline at end of file
+};
+
+/** Determines what a particle system will do when all of it's emitters become insignificant. */
+UENUM()
+enum class EParticleSystemInsignificanceReaction: uint8
+{
+	/** Looping systems will DisableTick. Non-looping systems will Complete.*/
+	Auto,
+	/** The system will be considered complete and will auto destroy if desired etc.*/
+	Complete,
+	/** The system will simply stop ticking. Tick will be re-enabled when any emitters become significant again. This is useful for persistent fx such as environmental fx.  */
+	DisableTick,
+	/** As DisableTick but will also kill all particles. */
+	DisableTickAndKill UMETA(Hidden), //Hidden for now until I make it useful i.e. Killing particles saves memory.
+
+	Num UMETA(Hidden),
+};
+
diff --git a/Engine/Source/Runtime/FriendsAndChat/Private/FriendsAndChatStyle.cpp b/Engine/Source/Runtime/FriendsAndChat/Private/FriendsAndChatStyle.cpp
--- a/Engine/Source/Runtime/FriendsAndChat/Private/FriendsAndChatStyle.cpp
+++ b/Engine/Source/Runtime/FriendsAndChat/Private/FriendsAndChatStyle.cpp
@@ -122,6 +122,10 @@ TSharedRef< FSlateStyleSet > FFriendsAndChatModuleStyle::Create(FFriendsAndChatS
 		.SetFont(FriendStyle.FriendsChatStyle.TextStyle.Font)
 		.SetColorAndOpacity(FriendStyle.FriendsChatStyle.GameHyperlinkChatColor);
 
+	const FTextBlockStyle TeamChatFont = FTextBlockStyle(DefaultText)
+		.SetFont(FriendStyle.FriendsChatStyle.TextStyle.Font)
+		.SetColorAndOpacity(FriendStyle.FriendsChatStyle.TeamHyperlinkChatColor);
+
 	const FTextBlockStyle PartyChatFont = FTextBlockStyle(DefaultText)
 		.SetFont(FriendStyle.FriendsChatStyle.TextStyle.Font)
 		.SetColorAndOpacity(FriendStyle.FriendsChatStyle.PartyHyperlinkChatColor);
@@ -145,6 +149,11 @@ TSharedRef< FSlateStyleSet > FFriendsAndChatModuleStyle::Create(FFriendsAndChatS
 		.SetTextStyle(GameChatFont)
 		.SetPadding(FMargin(0.0f));
 
+	const FHyperlinkStyle TeamChatHyperlink = FHyperlinkStyle()
+		.SetUnderlineStyle(UserNameButton)
+		.SetTextStyle(TeamChatFont)
+		.SetPadding(FMargin(0.0f));
+
 	const FHyperlinkStyle PartyChatHyperlink = FHyperlinkStyle()
 		.SetUnderlineStyle(UserNameButton)
 		.SetTextStyle(PartyChatFont)
@@ -164,11 +173,13 @@ TSharedRef< FSlateStyleSet > FFriendsAndChatModuleStyle::Create(FFriendsAndChatS
 
 	Style->Set("UserNameTextStyle.GlobalHyperlink", GlobalChatHyperlink);
 	Style->Set("UserNameTextStyle.GameHyperlink", GameChatHyperlink);
+	Style->Set("UserNameTextStyle.GameHyperlink", TeamChatHyperlink);
 	Style->Set("UserNameTextStyle.PartyHyperlink", PartyChatHyperlink);
 	Style->Set("UserNameTextStyle.Whisperlink", WhisperChatHyperlink);
 	Style->Set("UserNameTextStyle.DefaultHyperlink", DefaultChatHyperlink);
 	Style->Set("UserNameTextStyle.GlobalTextStyle", GlobalChatFont);
 	Style->Set("UserNameTextStyle.GameTextStyle", GameChatFont);
+	Style->Set("UserNameTextStyle.TeamTextStyle", TeamChatFont);
 	Style->Set("UserNameTextStyle.PartyTextStyle", PartyChatFont);
 	Style->Set("UserNameTextStyle.WhisperTextStyle", WhisperChatFont);
 
diff --git a/Engine/Source/Runtime/FriendsAndChat/Private/FriendsChatStyle.cpp b/Engine/Source/Runtime/FriendsAndChat/Private/FriendsChatStyle.cpp
--- a/Engine/Source/Runtime/FriendsAndChat/Private/FriendsChatStyle.cpp
+++ b/Engine/Source/Runtime/FriendsAndChat/Private/FriendsChatStyle.cpp
@@ -52,6 +52,12 @@ FFriendsChatStyle& FFriendsChatStyle::SetGameChatColor(const FLinearColor& InFon
 	return *this;
 }
 
+FFriendsChatStyle& FFriendsChatStyle::SetTeamChatColor(const FLinearColor& InFontColor)
+{
+	TeamChatColor = InFontColor;
+	return *this;
+}
+
 FFriendsChatStyle& FFriendsChatStyle::SetPartyChatColor(const FLinearColor& InFontColor)
 {
 	PartyChatColor = InFontColor;
@@ -64,6 +70,12 @@ FFriendsChatStyle& FFriendsChatStyle::SetAdminChatColor(const FLinearColor& InFo
 	return *this;
 }
 
+FFriendsChatStyle& FFriendsChatStyle::SetGameAdminChatColor(const FLinearColor& InFontColor)
+{
+	GameAdminChatColor = InFontColor;
+	return *this;
+}
+
 FFriendsChatStyle& FFriendsChatStyle::SetWhisplerHyperlinkChatColor(const FLinearColor& InFontColor)
 {
 	WhisperHyperlinkChatColor = InFontColor;
@@ -82,12 +94,30 @@ FFriendsChatStyle& FFriendsChatStyle::SetGameHyperlinkChatColor(const FLinearCol
 	return *this;
 }
 
+FFriendsChatStyle& FFriendsChatStyle::SetTeamHyperlinkChatColor(const FLinearColor& InFontColor)
+{
+	TeamHyperlinkChatColor = InFontColor;
+	return *this;
+}
+
 FFriendsChatStyle& FFriendsChatStyle::SetPartyHyperlinkChatColor(const FLinearColor& InFontColor)
 {
 	PartyHyperlinkChatColor = InFontColor;
 	return *this;
 }
 
+FFriendsChatStyle& FFriendsChatStyle::SetEnemyColor(const FLinearColor& InFontColor)
+{
+	EnemyColor = InFontColor;
+	return *this;
+}
+
+FFriendsChatStyle& FFriendsChatStyle::SetFriendlyColor(const FLinearColor& InFontColor)
+{
+	FriendlyColor = InFontColor;
+	return *this;
+}
+
 FFriendsChatStyle& FFriendsChatStyle::SetChatEntryTextStyle(const FEditableTextBoxStyle& InEditableTextStyle)
 {
 	ChatEntryTextStyle = InEditableTextStyle;
diff --git a/Engine/Source/Runtime/FriendsAndChat/Public/FriendsChatStyle.h b/Engine/Source/Runtime/FriendsAndChat/Public/FriendsChatStyle.h
--- a/Engine/Source/Runtime/FriendsAndChat/Public/FriendsChatStyle.h
+++ b/Engine/Source/Runtime/FriendsAndChat/Public/FriendsChatStyle.h
@@ -79,6 +79,10 @@ struct FRIENDSANDCHAT_API FFriendsChatStyle
 	FLinearColor GameChatColor;
 	FFriendsChatStyle& SetGameChatColor(const FLinearColor& InFontColor);
 
+	UPROPERTY(EditAnywhere, Category=Appearance)
+	FLinearColor TeamChatColor;
+	FFriendsChatStyle& SetTeamChatColor(const FLinearColor& InFontColor);
+
 	UPROPERTY(EditAnywhere, Category = Appearance)
 	FLinearColor PartyChatColor;
 	FFriendsChatStyle& SetPartyChatColor(const FLinearColor& InFontColor);
@@ -87,6 +91,10 @@ struct FRIENDSANDCHAT_API FFriendsChatStyle
 	FLinearColor AdminChatColor;
 	FFriendsChatStyle& SetAdminChatColor(const FLinearColor& InFontColor);
 
+	UPROPERTY(EditAnywhere, Category=Appearance)
+	FLinearColor GameAdminChatColor;
+	FFriendsChatStyle& SetGameAdminChatColor(const FLinearColor& InFontColor);
+
 	UPROPERTY(EditAnywhere, Category=Appearance)
 	FLinearColor WhisperHyperlinkChatColor;
 	FFriendsChatStyle& SetWhisplerHyperlinkChatColor(const FLinearColor& InFontColor);
@@ -99,10 +107,22 @@ struct FRIENDSANDCHAT_API FFriendsChatStyle
 	FLinearColor GameHyperlinkChatColor;
 	FFriendsChatStyle& SetGameHyperlinkChatColor(const FLinearColor& InFontColor);
 
+	UPROPERTY(EditAnywhere, Category=Appearance)
+	FLinearColor TeamHyperlinkChatColor;
+	FFriendsChatStyle& SetTeamHyperlinkChatColor(const FLinearColor& InFontColor);
+
 	UPROPERTY(EditAnywhere, Category = Appearance)
 	FLinearColor PartyHyperlinkChatColor;
 	FFriendsChatStyle& SetPartyHyperlinkChatColor(const FLinearColor& InFontColor);
 
+	UPROPERTY(EditAnywhere, Category = Appearance)
+	FLinearColor EnemyColor;
+	FFriendsChatStyle& SetEnemyColor(const FLinearColor& InFontColor);
+
+	UPROPERTY(EditAnywhere, Category = Appearance)
+	FLinearColor FriendlyColor;
+	FFriendsChatStyle& SetFriendlyColor(const FLinearColor& InFontColor);
+
 	UPROPERTY(EditAnywhere, Category = Appearance)
 	FEditableTextBoxStyle ChatEntryTextStyle;
 	FFriendsChatStyle& SetChatEntryTextStyle(const FEditableTextBoxStyle& InEditableTextStyle);
diff --git a/Engine/Source/Runtime/GameplayAbilities/Private/Abilities/GameplayAbility.cpp b/Engine/Source/Runtime/GameplayAbilities/Private/Abilities/GameplayAbility.cpp
--- a/Engine/Source/Runtime/GameplayAbilities/Private/Abilities/GameplayAbility.cpp
+++ b/Engine/Source/Runtime/GameplayAbilities/Private/Abilities/GameplayAbility.cpp
@@ -1622,6 +1622,28 @@ void UGameplayAbility::SetRemoteInstanceHasEnded()
 	}
 }
 
+void UGameplayAbility::NotifyAvatarDestroyed()
+{
+	// This could potentially happen in shutdown corner cases
+	if (IsPendingKill() || CurrentActorInfo == nullptr || CurrentActorInfo->AbilitySystemComponent.IsValid() == false)
+	{
+		return;
+	}
+
+	RemoteInstanceEnded = true;
+	for (UGameplayTask* Task : ActiveTasks)
+	{
+		if (Task && Task->IsPendingKill() == false && Task->IsWaitingOnAvatar())
+		{
+			// We have a task waiting on some Avatar state but the avatar is destroyed, so force end the ability to avoid getting stuck on.
+			
+			ABILITY_LOG(Log, TEXT("Ability %s is force cancelling because Task %s is waiting on avatar data avatar has been destroyed."), *GetName(), *Task->GetDebugString());
+			CurrentActorInfo->AbilitySystemComponent->ForceCancelAbilityDueToReplication(this);
+			break;
+		}
+	}
+}
+
 void UGameplayAbility::NotifyAbilityTaskWaitingOnPlayerData(class UAbilityTask* AbilityTask)
 {
 	// This should never happen since it will only be called from actively running ability tasks
@@ -1632,4 +1654,13 @@ void UGameplayAbility::NotifyAbilityTaskWaitingOnPlayerData(class UAbilityTask*
 		ABILITY_LOG(Log, TEXT("Ability %s is force cancelling because Task %s has started after the remote player has ended the ability."), *GetName(), *AbilityTask->GetDebugString());
 		CurrentActorInfo->AbilitySystemComponent->ForceCancelAbilityDueToReplication(this);
 	}
+}
+
+void UGameplayAbility::NotifyAbilityTaskWaitingOnAvatar(class UAbilityTask* AbilityTask)
+{
+	if (!CurrentActorInfo || CurrentActorInfo->AvatarActor.IsValid() == false)
+	{
+		ABILITY_LOG(Log, TEXT("Ability %s is force cancelling because Task %s has started while there is no valid AvatarActor"), *GetName(), *AbilityTask->GetDebugString());
+		CurrentActorInfo->AbilitySystemComponent->ForceCancelAbilityDueToReplication(this);
+	}
 }
\ No newline at end of file
diff --git a/Engine/Source/Runtime/GameplayAbilities/Private/Abilities/Tasks/AbilityTask.cpp b/Engine/Source/Runtime/GameplayAbilities/Private/Abilities/Tasks/AbilityTask.cpp
--- a/Engine/Source/Runtime/GameplayAbilities/Private/Abilities/Tasks/AbilityTask.cpp
+++ b/Engine/Source/Runtime/GameplayAbilities/Private/Abilities/Tasks/AbilityTask.cpp
@@ -7,7 +7,7 @@
 UAbilityTask::UAbilityTask(const FObjectInitializer& ObjectInitializer)
 	: Super(ObjectInitializer)
 {
-	WaitState = EAbilityTaskWaitState::WaitingOnGame;
+	WaitStateBitMask = (uint8)EAbilityTaskWaitState::WaitingOnGame;
 }
 
 FGameplayAbilitySpecHandle UAbilityTask::GetAbilitySpecHandle() const
@@ -61,17 +61,36 @@ void UAbilityTask::SetWaitingOnRemotePlayerData()
 {
 	if (Ability && IsPendingKill() == false && AbilitySystemComponent)
 	{
-		WaitState = EAbilityTaskWaitState::WaitingOnUser;
+		WaitStateBitMask |= (uint8)EAbilityTaskWaitState::WaitingOnUser;
 		Ability->NotifyAbilityTaskWaitingOnPlayerData(this);
 	}
 }
 
 void UAbilityTask::ClearWaitingOnRemotePlayerData()
 {
-	WaitState = EAbilityTaskWaitState::WaitingOnGame;
+	WaitStateBitMask &= ~((uint8)EAbilityTaskWaitState::WaitingOnUser);
 }
 
 bool UAbilityTask::IsWaitingOnRemotePlayerdata() const
 {
-	return (WaitState == EAbilityTaskWaitState::WaitingOnUser);
+	return (WaitStateBitMask & (uint8)EAbilityTaskWaitState::WaitingOnUser) != 0;
+}
+
+void UAbilityTask::SetWaitingOnAvatar()
+{
+	if (Ability && IsPendingKill() == false && AbilitySystemComponent)
+	{
+		WaitStateBitMask |= (uint8)EAbilityTaskWaitState::WaitingOnAvatar;
+		Ability->NotifyAbilityTaskWaitingOnAvatar(this);
+	}
+}
+
+void UAbilityTask::ClearWaitingOnAvatar()
+{
+	WaitStateBitMask &= ~((uint8)EAbilityTaskWaitState::WaitingOnAvatar);
+}
+
+bool UAbilityTask::IsWaitingOnAvatar() const
+{
+	return (WaitStateBitMask & (uint8)EAbilityTaskWaitState::WaitingOnAvatar) != 0;
 }
\ No newline at end of file
diff --git a/Engine/Source/Runtime/GameplayAbilities/Private/Abilities/Tasks/AbilityTask_ApplyRootMotionJumpForce.cpp b/Engine/Source/Runtime/GameplayAbilities/Private/Abilities/Tasks/AbilityTask_ApplyRootMotionJumpForce.cpp
--- a/Engine/Source/Runtime/GameplayAbilities/Private/Abilities/Tasks/AbilityTask_ApplyRootMotionJumpForce.cpp
+++ b/Engine/Source/Runtime/GameplayAbilities/Private/Abilities/Tasks/AbilityTask_ApplyRootMotionJumpForce.cpp
@@ -44,6 +44,7 @@ void UAbilityTask_ApplyRootMotionJumpForce::Activate()
 	{
 		Character->LandedDelegate.AddDynamic(this, &UAbilityTask_ApplyRootMotionJumpForce::OnLandedCallback);
 	}
+	SetWaitingOnAvatar();
 }
 
 void UAbilityTask_ApplyRootMotionJumpForce::OnLandedCallback(const FHitResult& Hit)
diff --git a/Engine/Source/Runtime/GameplayAbilities/Private/Abilities/Tasks/AbilityTask_PlayMontageAndWait.cpp b/Engine/Source/Runtime/GameplayAbilities/Private/Abilities/Tasks/AbilityTask_PlayMontageAndWait.cpp
--- a/Engine/Source/Runtime/GameplayAbilities/Private/Abilities/Tasks/AbilityTask_PlayMontageAndWait.cpp
+++ b/Engine/Source/Runtime/GameplayAbilities/Private/Abilities/Tasks/AbilityTask_PlayMontageAndWait.cpp
@@ -37,10 +37,8 @@ void UAbilityTask_PlayMontageAndWait::OnMontageBlendingOut(UAnimMontage* Montage
 	}
 	else
 	{
-		OnComplete.Broadcast();
+		OnBlendOut.Broadcast();
 	}
-
-	EndTask();
 }
 
 void UAbilityTask_PlayMontageAndWait::OnMontageInterrupted()
@@ -52,6 +50,16 @@ void UAbilityTask_PlayMontageAndWait::OnMontageInterrupted()
 	}
 }
 
+void UAbilityTask_PlayMontageAndWait::OnMontageEnded(UAnimMontage* Montage, bool bInterrupted)
+{
+	if (!bInterrupted)
+	{
+		OnCompleted.Broadcast();
+	}
+
+	EndTask();
+}
+
 UAbilityTask_PlayMontageAndWait* UAbilityTask_PlayMontageAndWait::CreatePlayMontageAndWaitProxy(UObject* WorldContextObject,
 	FName TaskInstanceName, UAnimMontage *MontageToPlay, float Rate, FName StartSection, bool bStopWhenAbilityEnds, float AnimRootMotionTranslationScale)
 {
@@ -91,7 +99,10 @@ void UAbilityTask_PlayMontageAndWait::Activate()
 				InterruptedHandle = Ability->OnGameplayAbilityCancelled.AddUObject(this, &UAbilityTask_PlayMontageAndWait::OnMontageInterrupted);
 
 				BlendingOutDelegate.BindUObject(this, &UAbilityTask_PlayMontageAndWait::OnMontageBlendingOut);
-				ActorInfo->AnimInstance->Montage_SetBlendingOutDelegate(BlendingOutDelegate);
+				ActorInfo->AnimInstance->Montage_SetBlendingOutDelegate(BlendingOutDelegate, MontageToPlay);
+
+				MontageEndedDelegate.BindUObject(this, &UAbilityTask_PlayMontageAndWait::OnMontageEnded);
+				ActorInfo->AnimInstance->Montage_SetEndDelegate(MontageEndedDelegate, MontageToPlay);
 
 				ACharacter* Character = Cast<ACharacter>(GetAvatarActor());
 				if (Character && (Character->Role == ROLE_Authority ||
@@ -110,13 +121,14 @@ void UAbilityTask_PlayMontageAndWait::Activate()
 		ABILITY_LOG(Warning, TEXT("UAbilityTask_PlayMontageAndWait called in Ability %s failed to play montage; Task Instance Name %s."), *Ability->GetName(), *InstanceName.ToString());
 		OnCancelled.Broadcast();
 	}
+
+	SetWaitingOnAvatar();
 }
 
 void UAbilityTask_PlayMontageAndWait::ExternalCancel()
 {
 	check(AbilitySystemComponent);
 
-
 	OnCancelled.Broadcast();
 	Super::ExternalCancel();
 }
@@ -156,9 +168,11 @@ bool UAbilityTask_PlayMontageAndWait::StopPlayingMontage()
 			&& AbilitySystemComponent->GetCurrentMontage() == MontageToPlay)
 		{
 			// Unbind delegates so they don't get called as well
-			if (FOnMontageBlendingOutStarted* BoundDelegate = ActorInfo->AnimInstance->Montage_GetBlendingOutDelegate(MontageToPlay))
+			FAnimMontageInstance* MontageInstance = ActorInfo->AnimInstance->GetActiveInstanceForMontage(*MontageToPlay);
+			if (MontageInstance)
 			{
-				BoundDelegate->Unbind();
+				MontageInstance->OnMontageBlendingOutStarted.Unbind();
+				MontageInstance->OnMontageEnded.Unbind();
 			}
 
 			AbilitySystemComponent->CurrentMontageStop();
diff --git a/Engine/Source/Runtime/GameplayAbilities/Private/Abilities/Tasks/AbilityTask_WaitMovementModeChange.cpp b/Engine/Source/Runtime/GameplayAbilities/Private/Abilities/Tasks/AbilityTask_WaitMovementModeChange.cpp
--- a/Engine/Source/Runtime/GameplayAbilities/Private/Abilities/Tasks/AbilityTask_WaitMovementModeChange.cpp
+++ b/Engine/Source/Runtime/GameplayAbilities/Private/Abilities/Tasks/AbilityTask_WaitMovementModeChange.cpp
@@ -24,6 +24,8 @@ void UAbilityTask_WaitMovementModeChange::Activate()
 		Character->MovementModeChangedDelegate.AddDynamic(this, &UAbilityTask_WaitMovementModeChange::OnMovementModeChange);
 		MyCharacter = Character;
 	}
+
+	SetWaitingOnAvatar();
 }
 
 void UAbilityTask_WaitMovementModeChange::OnMovementModeChange(ACharacter * Character, EMovementMode PrevMovementMode, uint8 PreviousCustomMode)
diff --git a/Engine/Source/Runtime/GameplayAbilities/Private/Abilities/Tasks/AbilityTask_WaitOverlap.cpp b/Engine/Source/Runtime/GameplayAbilities/Private/Abilities/Tasks/AbilityTask_WaitOverlap.cpp
--- a/Engine/Source/Runtime/GameplayAbilities/Private/Abilities/Tasks/AbilityTask_WaitOverlap.cpp
+++ b/Engine/Source/Runtime/GameplayAbilities/Private/Abilities/Tasks/AbilityTask_WaitOverlap.cpp
@@ -49,6 +49,8 @@ UAbilityTask_WaitOverlap* UAbilityTask_WaitOverlap::WaitForOverlap(UObject* Worl
 
 void UAbilityTask_WaitOverlap::Activate()
 {
+	SetWaitingOnAvatar();
+
 	UPrimitiveComponent* PrimComponent = GetComponent();
 	if (PrimComponent)
 	{
diff --git a/Engine/Source/Runtime/GameplayAbilities/Private/Abilities/Tasks/AbilityTask_WaitVelocityChange.cpp b/Engine/Source/Runtime/GameplayAbilities/Private/Abilities/Tasks/AbilityTask_WaitVelocityChange.cpp
--- a/Engine/Source/Runtime/GameplayAbilities/Private/Abilities/Tasks/AbilityTask_WaitVelocityChange.cpp
+++ b/Engine/Source/Runtime/GameplayAbilities/Private/Abilities/Tasks/AbilityTask_WaitVelocityChange.cpp
@@ -44,4 +44,5 @@ void UAbilityTask_WaitVelocityChange::Activate()
 {
 	const FGameplayAbilityActorInfo* ActorInfo = Ability->GetCurrentActorInfo();
 	CachedMovementComponent = ActorInfo->MovementComponent.Get();
+	SetWaitingOnAvatar();
 }
\ No newline at end of file
diff --git a/Engine/Source/Runtime/GameplayAbilities/Private/AbilitySystemComponent.cpp b/Engine/Source/Runtime/GameplayAbilities/Private/AbilitySystemComponent.cpp
--- a/Engine/Source/Runtime/GameplayAbilities/Private/AbilitySystemComponent.cpp
+++ b/Engine/Source/Runtime/GameplayAbilities/Private/AbilitySystemComponent.cpp
@@ -1843,8 +1843,8 @@ void UAbilitySystemComponent::Debug_Internal(FAbilitySystemComponentDebugInfo& I
 					int32 MsgCount = 0;
 					for (FAbilityTaskDebugMessage& Msg : Instance->TaskDebugMessages)
 					{
-						// Cap finished task msgs to 5 per ability unless dumping to log
-						if ( ++MsgCount > 5 || Info.bPrintToLog)
+						// Cap finished task msgs to 5 per ability if we are printing to screen (else things will scroll off)
+						if ( Info.Canvas && ++MsgCount > 5 )
 						{
 							break;
 						}
diff --git a/Engine/Source/Runtime/GameplayAbilities/Private/AttributeSet.cpp b/Engine/Source/Runtime/GameplayAbilities/Private/AttributeSet.cpp
--- a/Engine/Source/Runtime/GameplayAbilities/Private/AttributeSet.cpp
+++ b/Engine/Source/Runtime/GameplayAbilities/Private/AttributeSet.cpp
@@ -152,12 +152,23 @@ float FScalableFloat::GetValueAtLevel(float Level) const
 {
 	if (Curve.CurveTable != nullptr)
 	{
-		if (FinalCurve == nullptr)
+		// We want to cache the FinalCurve when possible. However there are issue with this in CookOnTheFly scenerios where the backing UCurveTable
+		// will be destroyed and reallocated at the same memory adress via StaticAllocateObject. This will keep the UCurveTable* the same, but will shift
+		// the internal UCurveTable::RowMap memory, making our FinalCurve* into it invalid.
+		//
+		// Therefor: only do the caching in non editor builds. In editor builds we will also look up the FinalCurve* from our current UCurveTable*.
+		// 
+		// If this problems does happen in a mono build, we should consider caching the entire curve in the FScalableFloat, or have look into refactoring
+		// some of UCurveTable such that the actual curve data is stored in a global, more stable location.
+
+
+#if !WITH_EDITOR						// If not an editor build
+		if (FinalCurve == nullptr)		//		Only if we don't have a cached FRichCurve*, look it up
+#endif
+		
 		{
 			static const FString ContextString = TEXT("FScalableFloat::GetValueAtLevel");
 			FinalCurve = Curve.GetCurve(ContextString);
-
-			RegisterOnCurveTablePostReimport();
 		}
 
 		if (FinalCurve != nullptr)
diff --git a/Engine/Source/Runtime/GameplayAbilities/Private/GameplayEffect.cpp b/Engine/Source/Runtime/GameplayAbilities/Private/GameplayEffect.cpp
--- a/Engine/Source/Runtime/GameplayAbilities/Private/GameplayEffect.cpp
+++ b/Engine/Source/Runtime/GameplayAbilities/Private/GameplayEffect.cpp
@@ -2668,7 +2668,11 @@ bool FActiveGameplayEffectsContainer::InternalRemoveActiveGameplayEffect(int32 I
 	if (ensure(Idx < GetNumGameplayEffects()))
 	{
 		FActiveGameplayEffect& Effect = *GetActiveGameplayEffect(Idx);
-		ensure(!Effect.IsPendingRemove);
+		if (!ensure(!Effect.IsPendingRemove))
+		{
+			// This effect is already being removed. This probably means a bug at the callsite, but we can handle it gracefully here by earlying out and pretending the effect was removed.
+			return true;
+		}
 
 		ABILITY_LOG(Verbose, TEXT("InternalRemoveActiveGameplayEffect: Auth: %s Handle: %s Def: %s"), IsNetAuthority() ? TEXT("TRUE") : TEXT("FALSE"), *Effect.Handle.ToString(), Effect.Spec.Def ? *Effect.Spec.Def->GetName() : TEXT("NONE"));
 
@@ -3053,8 +3057,15 @@ void FActiveGameplayEffectsContainer::CheckDuration(FActiveGameplayEffectHandle
 	for (int32 ActiveGEIdx = 0; ActiveGEIdx < GameplayEffects_Internal.Num(); ++ActiveGEIdx)
 	{
 		FActiveGameplayEffect& Effect = GameplayEffects_Internal[ActiveGEIdx];
-		if (Effect.Handle == Handle && Effect.IsPendingRemove == false)
+		if (Effect.Handle == Handle)
 		{
+			if (Effect.IsPendingRemove)
+			{
+				// break is this effect is pending remove. 
+				// (Note: don't combine this with the above if statement that is looking for the effect via handle, since we want to stop iteration if we find a matching handle but are pending remove).
+				break;
+			}
+
 			FTimerManager& TimerManager = Owner->GetWorld()->GetTimerManager();
 
 			// The duration may have changed since we registered this callback with the timer manager.
@@ -3104,6 +3115,14 @@ void FActiveGameplayEffectsContainer::CheckDuration(FActiveGameplayEffectHandle
 					if (PeriodTimeRemaining <= KINDA_SMALL_NUMBER && !Effect.bIsInhibited)
 					{
 						ExecuteActiveEffectsFrom(Effect.Spec);
+
+						// The above call to ExecuteActiveEffectsFrom could cause this effect to be explicitly removed
+						// (for example it could kill the owner and cause the effect to be wiped via death).
+						// In that case, we need to early out instead of possibly continueing to the below calls to InternalRemoveActiveGameplayEffect
+						if ( Effect.IsPendingRemove )
+						{
+							break;
+						}
 					}
 
 					// Forcibly clear the periodic ticks because this effect is going to be removed
diff --git a/Engine/Source/Runtime/GameplayAbilities/Public/Abilities/GameplayAbility.h b/Engine/Source/Runtime/GameplayAbilities/Public/Abilities/GameplayAbility.h
--- a/Engine/Source/Runtime/GameplayAbilities/Public/Abilities/GameplayAbility.h
+++ b/Engine/Source/Runtime/GameplayAbilities/Public/Abilities/GameplayAbility.h
@@ -629,8 +629,13 @@ public:
 	/** Called by ability system component to inform this ability instance the remote instance was ended */
 	void SetRemoteInstanceHasEnded();
 
+	/** Called to inform the ability that the AvatarActor has been replaced. If the ability is dependant on avatar state, it may want to end itself. */
+	void NotifyAvatarDestroyed();
+
 	void NotifyAbilityTaskWaitingOnPlayerData(class UAbilityTask* AbilityTask);
 
+	void NotifyAbilityTaskWaitingOnAvatar(class UAbilityTask* AbilityTask);
+
 	UFUNCTION(BlueprintCallable, Category = Ability)
 	float GetCooldownTimeRemaining() const;
 
diff --git a/Engine/Source/Runtime/GameplayAbilities/Public/Abilities/Tasks/AbilityTask.h b/Engine/Source/Runtime/GameplayAbilities/Public/Abilities/Tasks/AbilityTask.h
--- a/Engine/Source/Runtime/GameplayAbilities/Public/Abilities/Tasks/AbilityTask.h
+++ b/Engine/Source/Runtime/GameplayAbilities/Public/Abilities/Tasks/AbilityTask.h
@@ -67,10 +67,13 @@ UENUM()
 enum class EAbilityTaskWaitState : uint8
 {
 	/** Task is waiting for the game to do something */
-	WaitingOnGame,
+	WaitingOnGame = 0x01,
 
 	/** Waiting for the user to do something */
-	WaitingOnUser,
+	WaitingOnUser = 0x02,
+
+	/** Waiting on Avatar (Character/Pawn/Actor) to do something (usually something physical in the world, like land, move, etc) */
+	WaitingOnAvatar = 0x04
 };
 
 UCLASS(Abstract)
@@ -133,11 +136,15 @@ class GAMEPLAYABILITIES_API UAbilityTask : public UGameplayTask
 	/** Called when the ability task is waiting on remote player data. IF the remote player ends the ability prematurely, and a task with this set is still running, the ability is killed. */
 	void SetWaitingOnRemotePlayerData();
 	void ClearWaitingOnRemotePlayerData();
-
 	virtual bool IsWaitingOnRemotePlayerdata() const override;
 
+	/** same as RemotePlayerData but for ACharacter type of state (movement state, etc) */
+	void SetWaitingOnAvatar();
+	void ClearWaitingOnAvatar();
+	virtual bool IsWaitingOnAvatar() const override;
+
 	/** What we are waiting on */
-	EAbilityTaskWaitState WaitState;
+	uint8 WaitStateBitMask;
 
 protected:
 	/** Helper method for registering client replicated callbacks */
diff --git a/Engine/Source/Runtime/GameplayAbilities/Public/Abilities/Tasks/AbilityTask_PlayMontageAndWait.h b/Engine/Source/Runtime/GameplayAbilities/Public/Abilities/Tasks/AbilityTask_PlayMontageAndWait.h
--- a/Engine/Source/Runtime/GameplayAbilities/Public/Abilities/Tasks/AbilityTask_PlayMontageAndWait.h
+++ b/Engine/Source/Runtime/GameplayAbilities/Public/Abilities/Tasks/AbilityTask_PlayMontageAndWait.h
@@ -11,7 +11,10 @@ class UAbilityTask_PlayMontageAndWait : public UAbilityTask
 	GENERATED_UCLASS_BODY()
 
 	UPROPERTY(BlueprintAssignable)
-	FMontageWaitSimpleDelegate	OnComplete;
+	FMontageWaitSimpleDelegate	OnCompleted;
+
+	UPROPERTY(BlueprintAssignable)
+	FMontageWaitSimpleDelegate	OnBlendOut;
 
 	UPROPERTY(BlueprintAssignable)
 	FMontageWaitSimpleDelegate	OnInterrupted;
@@ -25,6 +28,9 @@ class UAbilityTask_PlayMontageAndWait : public UAbilityTask
 	UFUNCTION()
 	void OnMontageInterrupted();
 
+	UFUNCTION()
+	void OnMontageEnded(UAnimMontage* Montage, bool bInterrupted);
+
 	UFUNCTION(BlueprintCallable, Category="Ability|Tasks", meta = (DisplayName="PlayMontageAndWait",
 		HidePin = "WorldContextObject", DefaultToSelf = "WorldContextObject", BlueprintInternalUseOnly = "TRUE"))
 	static UAbilityTask_PlayMontageAndWait* CreatePlayMontageAndWaitProxy(UObject* WorldContextObject,
@@ -45,6 +51,7 @@ private:
 	bool StopPlayingMontage();
 
 	FOnMontageBlendingOutStarted BlendingOutDelegate;
+	FOnMontageEnded MontageEndedDelegate;
 	FDelegateHandle InterruptedHandle;
 
 	UAnimMontage* MontageToPlay;
diff --git a/Engine/Source/Runtime/GameplayTasks/Classes/GameplayTask.h b/Engine/Source/Runtime/GameplayTasks/Classes/GameplayTask.h
--- a/Engine/Source/Runtime/GameplayTasks/Classes/GameplayTask.h
+++ b/Engine/Source/Runtime/GameplayTasks/Classes/GameplayTask.h
@@ -258,6 +258,8 @@ public:
 
 	virtual bool IsWaitingOnRemotePlayerdata() const { return false; }
 
+	virtual bool IsWaitingOnAvatar() const { return false; }
+
 protected:
 	/** End and CleanUp the task - may be called by the task itself or by the task owner if the owner is ending. 
 	 *	IMPORTANT! Do NOT call directly! Call EndTask() or TaskOwnerEnded() 
diff --git a/Engine/Source/Runtime/NetworkReplayStreaming/HttpNetworkReplayStreaming/Private/HttpNetworkReplayStreaming.cpp b/Engine/Source/Runtime/NetworkReplayStreaming/HttpNetworkReplayStreaming/Private/HttpNetworkReplayStreaming.cpp
--- a/Engine/Source/Runtime/NetworkReplayStreaming/HttpNetworkReplayStreaming/Private/HttpNetworkReplayStreaming.cpp
+++ b/Engine/Source/Runtime/NetworkReplayStreaming/HttpNetworkReplayStreaming/Private/HttpNetworkReplayStreaming.cpp
@@ -305,18 +305,18 @@ void FHttpNetworkReplayStreamer::StartStreaming( const FString& CustomName, cons
 			HttpRequest->SetContentAsString( UserList.ToJson() );
 		}
 
-		AddRequestToQueue( EQueuedHttpRequestType::StartUploading, HttpRequest );
+		AddRequestToQueue( EQueuedHttpRequestType::StartUploading, HttpRequest, 3, 2.0f );
 		
 		// We need to upload the header AFTER StartUploading is done (so we have session name)
 		AddRequestToQueue( EQueuedHttpRequestType::UploadHeader, nullptr );
 	}
 }
 
-void FHttpNetworkReplayStreamer::AddRequestToQueue( const EQueuedHttpRequestType::Type Type, TSharedPtr< class IHttpRequest > Request )
+void FHttpNetworkReplayStreamer::AddRequestToQueue( const EQueuedHttpRequestType::Type Type, TSharedPtr< class IHttpRequest > Request, const int32 InMaxRetries, const float InRetryDelay )
 {
 	UE_LOG( LogHttpReplay, Verbose, TEXT( "FHttpNetworkReplayStreamer::AddRequestToQueue. Type: %s" ), EQueuedHttpRequestType::ToString( Type ) );
 
-	QueuedHttpRequests.Add( TSharedPtr< FQueuedHttpRequest >( new FQueuedHttpRequest( Type, Request ) ) );
+	QueuedHttpRequests.Add( TSharedPtr< FQueuedHttpRequest >( new FQueuedHttpRequest( Type, Request, InMaxRetries, InRetryDelay ) ) );
 }
 
 void FHttpNetworkReplayStreamer::AddCustomRequestToQueue( TSharedPtr< FQueuedHttpRequest > Request )
@@ -326,6 +326,73 @@ void FHttpNetworkReplayStreamer::AddCustomRequestToQueue( TSharedPtr< FQueuedHtt
 	QueuedHttpRequests.Add( Request );
 }
 
+static FString BuildRequestErrorString( FHttpRequestPtr HttpRequest, FHttpResponsePtr HttpResponse )
+{
+	FString ExtraInfo;
+
+	if ( HttpRequest.IsValid() )
+	{
+		ExtraInfo += FString::Printf( TEXT( "URL: %s" ), *HttpRequest->GetURL() );
+		ExtraInfo += FString::Printf( TEXT( ", Verb: %s" ), *HttpRequest->GetVerb() );
+
+		const TArray< FString > AllHeaders = HttpRequest->GetAllHeaders();
+
+		for ( int i = 0; i < AllHeaders.Num(); i++ )
+		{
+			ExtraInfo += TEXT( ", " );
+			ExtraInfo += AllHeaders[i];
+		}
+	}
+	else
+	{
+		ExtraInfo = TEXT( "HttpRequest NULL." );
+	}
+
+	return FString::Printf( TEXT( "Response code: %d, Extra info: %s" ), HttpResponse.IsValid() ? HttpResponse->GetResponseCode() : 0, *ExtraInfo );
+}
+
+bool FHttpNetworkReplayStreamer::RetryRequest( TSharedPtr< FQueuedHttpRequest > Request, FHttpResponsePtr HttpResponse )
+{
+	if ( !Request.IsValid() )
+	{
+		return false;
+	}
+
+	if ( Request->MaxRetries == 0 || Request->RetryProgress >= Request->MaxRetries )
+	{
+		return false;
+	}
+	
+	if ( !HttpResponse.IsValid() )
+	{
+		return false;
+	}
+
+	if ( HttpResponse->GetResponseCode() < 500 || HttpResponse->GetResponseCode() >= 600 )
+	{
+		return false;		// Only retry on 5xx return codes
+	}
+
+	Request->RetryProgress++;
+
+	if ( Request->RetryDelay == 0.0f )
+	{
+		// Re-process it immediately if it's not on a delay
+		Request->Request->ProcessRequest();
+	}
+	else
+	{
+		Request->NextRetryTime = FPlatformTime::Seconds() + Request->RetryDelay;
+	}
+
+	// Put request back on InFlightRequest
+	InFlightHttpRequest = Request;
+
+	UE_LOG( LogHttpReplay, Warning, TEXT( "FHttpNetworkReplayStreamer::RetryRequest. Retrying: %i, %s" ), Request->RetryProgress, *BuildRequestErrorString( Request->Request, HttpResponse ) );
+
+	return true;
+}
+
 void FHttpNetworkReplayStreamer::StopStreaming()
 {
 	if ( StartStreamingDelegate.IsBound() )
@@ -405,7 +472,7 @@ void FHttpNetworkReplayStreamer::UploadHeader()
 	HeaderArchive.Buffer.Empty();
 	HeaderArchive.Pos = 0;
 
-	AddRequestToQueue( EQueuedHttpRequestType::UploadingHeader, HttpRequest );
+	AddRequestToQueue( EQueuedHttpRequestType::UploadingHeader, HttpRequest, 3, 2.0f );
 
 	LastChunkTime = FPlatformTime::Seconds();
 }
@@ -471,7 +538,7 @@ void FHttpNetworkReplayStreamer::FlushStream()
 
 	StreamChunkIndex++;
 
-	AddRequestToQueue( EQueuedHttpRequestType::UploadingStream, HttpRequest );
+	AddRequestToQueue( EQueuedHttpRequestType::UploadingStream, HttpRequest, 2, 2.0f );
 
 	LastChunkTime = FPlatformTime::Seconds();
 }
@@ -502,7 +569,7 @@ void FHttpNetworkReplayStreamer::StopUploading()
 	HttpRequest->SetVerb( TEXT( "POST" ) );
 	HttpRequest->SetHeader( TEXT( "Content-Type" ), TEXT( "application/octet-stream" ) );
 
-	AddRequestToQueue( EQueuedHttpRequestType::StopUploading, HttpRequest );
+	AddRequestToQueue( EQueuedHttpRequestType::StopUploading, HttpRequest, 3, 2.0f );
 };
 
 void FHttpNetworkReplayStreamer::FlushCheckpoint( const uint32 TimeInMS )
@@ -730,7 +797,7 @@ void FHttpNetworkReplayStreamer::FlushCheckpointInternal( uint32 TimeInMS )
 	CheckpointArchive.Buffer.Empty();
 	CheckpointArchive.Pos = 0;
 
-	AddRequestToQueue( EQueuedHttpRequestType::UploadingCheckpoint, HttpRequest );
+	AddRequestToQueue( EQueuedHttpRequestType::UploadingCheckpoint, HttpRequest, 2, 2.0f );
 }
 
 FQueuedHttpRequestAddEvent::FQueuedHttpRequestAddEvent( const FString& InName, const uint32 InTimeInMS, const FString& InGroup, const FString& InMeta, const TArray<uint8>& InData, TSharedRef< class IHttpRequest > InHttpRequest ) : FQueuedHttpRequest( EQueuedHttpRequestType::UploadingCustomEvent, InHttpRequest )
@@ -1293,33 +1360,10 @@ void FHttpNetworkReplayStreamer::RequestFinished( EStreamerState ExpectedStreame
 	InFlightHttpRequest = NULL;
 };
 
-static FString BuildRequestErrorString( FHttpRequestPtr HttpRequest, FHttpResponsePtr HttpResponse )
-{
-	FString ExtraInfo;
-
-	if ( HttpRequest.IsValid() )
-	{
-		ExtraInfo += FString::Printf( TEXT( "URL: %s" ), *HttpRequest->GetURL() );
-		ExtraInfo += FString::Printf( TEXT( ", Verb: %s" ), *HttpRequest->GetVerb() );
-
-		const TArray< FString > AllHeaders = HttpRequest->GetAllHeaders();
-
-		for ( int i = 0; i < AllHeaders.Num(); i++ )
-		{
-			ExtraInfo += TEXT( ", " );
-			ExtraInfo += AllHeaders[i];
-		}
-	}
-	else
-	{
-		ExtraInfo = TEXT( "HttpRequest NULL." );
-	}
-
-	return FString::Printf( TEXT( "Response code: %d, Extra info: %s" ), HttpResponse.IsValid() ? HttpResponse->GetResponseCode() : 0, *ExtraInfo );
-}
-
 void FHttpNetworkReplayStreamer::HttpStartUploadingFinished( FHttpRequestPtr HttpRequest, FHttpResponsePtr HttpResponse, bool bSucceeded )
 {
+	TSharedPtr< FQueuedHttpRequest > SavedFlightHttpRequest = InFlightHttpRequest;
+
 	RequestFinished( EStreamerState::StreamingUp, EQueuedHttpRequestType::StartUploading, HttpRequest );
 
 	if ( bSucceeded && HttpResponse->GetResponseCode() == EHttpResponseCodes::Ok )
@@ -1341,6 +1385,11 @@ void FHttpNetworkReplayStreamer::HttpStartUploadingFinished( FHttpRequestPtr Htt
 	}
 	else
 	{
+		if ( RetryRequest( SavedFlightHttpRequest, HttpResponse ) )
+		{
+			return;
+		}
+
 		UE_LOG( LogHttpReplay, Error, TEXT( "FHttpNetworkReplayStreamer::HttpStartUploadingFinished. FAILED, %s" ), *BuildRequestErrorString( HttpRequest, HttpResponse ) );
 		SetLastError( ENetworkReplayError::ServiceUnavailable );
 	}
@@ -1348,6 +1397,8 @@ void FHttpNetworkReplayStreamer::HttpStartUploadingFinished( FHttpRequestPtr Htt
 
 void FHttpNetworkReplayStreamer::HttpStopUploadingFinished( FHttpRequestPtr HttpRequest, FHttpResponsePtr HttpResponse, bool bSucceeded )
 {
+	TSharedPtr< FQueuedHttpRequest > SavedFlightHttpRequest = InFlightHttpRequest;
+
 	RequestFinished( EStreamerState::StreamingUp, EQueuedHttpRequestType::StopUploading, HttpRequest );
 
 	if ( bSucceeded && HttpResponse->GetResponseCode() == EHttpResponseCodes::NoContent )
@@ -1356,6 +1407,11 @@ void FHttpNetworkReplayStreamer::HttpStopUploadingFinished( FHttpRequestPtr Http
 	}
 	else
 	{
+		if ( RetryRequest( SavedFlightHttpRequest, HttpResponse ) )
+		{
+			return;
+		}
+
 		UE_LOG( LogHttpReplay, Error, TEXT( "FHttpNetworkReplayStreamer::HttpStopUploadingFinished. FAILED, %s" ), *BuildRequestErrorString( HttpRequest, HttpResponse ) );
 		SetLastError( ENetworkReplayError::ServiceUnavailable );
 	}
@@ -1370,6 +1426,8 @@ void FHttpNetworkReplayStreamer::HttpStopUploadingFinished( FHttpRequestPtr Http
 
 void FHttpNetworkReplayStreamer::HttpHeaderUploadFinished( FHttpRequestPtr HttpRequest, FHttpResponsePtr HttpResponse, bool bSucceeded )
 {
+	TSharedPtr< FQueuedHttpRequest > SavedFlightHttpRequest = InFlightHttpRequest;
+
 	RequestFinished( EStreamerState::StreamingUp, EQueuedHttpRequestType::UploadingHeader, HttpRequest );
 
 	check( StartStreamingDelegate.IsBound() );
@@ -1382,6 +1440,11 @@ void FHttpNetworkReplayStreamer::HttpHeaderUploadFinished( FHttpRequestPtr HttpR
 	}
 	else
 	{
+		if ( RetryRequest( SavedFlightHttpRequest, HttpResponse ) )
+		{
+			return;
+		}
+
 		UE_LOG( LogHttpReplay, Error, TEXT( "FHttpNetworkReplayStreamer::HttpHeaderUploadFinished. FAILED, %s" ), *BuildRequestErrorString( HttpRequest, HttpResponse ) );
 		StartStreamingDelegate.ExecuteIfBound( false, true );
 		SetLastError( ENetworkReplayError::ServiceUnavailable );
@@ -1393,6 +1456,8 @@ void FHttpNetworkReplayStreamer::HttpHeaderUploadFinished( FHttpRequestPtr HttpR
 
 void FHttpNetworkReplayStreamer::HttpUploadStreamFinished( FHttpRequestPtr HttpRequest, FHttpResponsePtr HttpResponse, bool bSucceeded )
 {
+	TSharedPtr< FQueuedHttpRequest > SavedFlightHttpRequest = InFlightHttpRequest;
+
 	RequestFinished( EStreamerState::StreamingUp, EQueuedHttpRequestType::UploadingStream, HttpRequest );
 
 	if ( bSucceeded && HttpResponse->GetResponseCode() == EHttpResponseCodes::NoContent )
@@ -1401,6 +1466,11 @@ void FHttpNetworkReplayStreamer::HttpUploadStreamFinished( FHttpRequestPtr HttpR
 	}
 	else
 	{
+		if ( RetryRequest( SavedFlightHttpRequest, HttpResponse ) )
+		{
+			return;
+		}
+
 		UE_LOG( LogHttpReplay, Error, TEXT( "FHttpNetworkReplayStreamer::HttpUploadStreamFinished. FAILED, %s" ), *BuildRequestErrorString( HttpRequest, HttpResponse ) );
 		SetLastError( ENetworkReplayError::ServiceUnavailable );
 	}
@@ -1408,6 +1478,8 @@ void FHttpNetworkReplayStreamer::HttpUploadStreamFinished( FHttpRequestPtr HttpR
 
 void FHttpNetworkReplayStreamer::HttpUploadCheckpointFinished( FHttpRequestPtr HttpRequest, FHttpResponsePtr HttpResponse, bool bSucceeded )
 {
+	TSharedPtr< FQueuedHttpRequest > SavedFlightHttpRequest = InFlightHttpRequest;
+
 	RequestFinished( EStreamerState::StreamingUp, EQueuedHttpRequestType::UploadingCheckpoint, HttpRequest );
 
 	if ( bSucceeded && ( HttpResponse->GetResponseCode() == EHttpResponseCodes::Ok || HttpResponse->GetResponseCode() == EHttpResponseCodes::NoContent ) )
@@ -1416,7 +1488,12 @@ void FHttpNetworkReplayStreamer::HttpUploadCheckpointFinished( FHttpRequestPtr H
 	}
 	else
 	{
-		UE_LOG( LogHttpReplay, Error, TEXT( "FHttpNetworkReplayStreamer::HttpUploadCheckpointFinished. FAILED, %s" ), *BuildRequestErrorString( HttpRequest, HttpResponse ) );
+		if ( RetryRequest( SavedFlightHttpRequest, HttpResponse ) )
+		{
+			return;
+		}
+
+		UE_LOG( LogHttpReplay, Warning, TEXT( "FHttpNetworkReplayStreamer::HttpUploadCheckpointFinished. FAILED, %s" ), *BuildRequestErrorString( HttpRequest, HttpResponse ) );
 		SetLastError( ENetworkReplayError::ServiceUnavailable );
 	}
 }
@@ -1431,29 +1508,7 @@ void FHttpNetworkReplayStreamer::HttpUploadCustomEventFinished(FHttpRequestPtr H
 	}
 	else
 	{
-		FString ExtraInfo;
-
-		if ( HttpRequest.IsValid() )
-		{
-			ExtraInfo += FString::Printf( TEXT( "URL: %s" ),	*HttpRequest->GetURL() );
-			ExtraInfo += FString::Printf( TEXT( ", Verb: %s" ), *HttpRequest->GetVerb() );
-
-			const TArray< FString > AllHeaders = HttpRequest->GetAllHeaders();
-
-			for ( int i = 0; i < AllHeaders.Num(); i++ )
-			{
-				ExtraInfo += TEXT( ", " );
-				ExtraInfo += AllHeaders[i];
-			}
-		}
-		else
-		{
-			ExtraInfo = TEXT( "HttpRequest NULL." );
-		}
-
-		UE_LOG( LogHttpReplay, Error, TEXT( "FHttpNetworkReplayStreamer::HttpUploadCustomEventFinished. FAILED. Response code: %d, Extra info: %s" ), HttpResponse.IsValid() ? HttpResponse->GetResponseCode() : 0, *ExtraInfo );
-		// Don't disconect service here, just report the failure
-		//SetLastError(ENetworkReplayError::ServiceUnavailable);
+		UE_LOG( LogHttpReplay, Warning, TEXT( "FHttpNetworkReplayStreamer::HttpUploadCustomEventFinished. FAILED. Extra info: %s" ), *BuildRequestErrorString( HttpRequest, HttpResponse ) );
 	}
 }
 
@@ -1906,6 +1961,15 @@ bool FHttpNetworkReplayStreamer::ProcessNextHttpRequest()
 {
 	if ( IsHttpRequestInFlight() )
 	{
+		if ( InFlightHttpRequest->NextRetryTime > 0.0f )
+		{
+			if ( FPlatformTime::Seconds() > InFlightHttpRequest->NextRetryTime )
+			{
+				InFlightHttpRequest->NextRetryTime = 0.0f;
+				InFlightHttpRequest->Request->ProcessRequest();
+			}
+		}
+
 		// We only process one http request at a time to keep things simple
 		return false;
 	}
diff --git a/Engine/Source/Runtime/NetworkReplayStreaming/HttpNetworkReplayStreaming/Public/HttpNetworkReplayStreaming.h b/Engine/Source/Runtime/NetworkReplayStreaming/HttpNetworkReplayStreaming/Public/HttpNetworkReplayStreaming.h
--- a/Engine/Source/Runtime/NetworkReplayStreaming/HttpNetworkReplayStreaming/Public/HttpNetworkReplayStreaming.h
+++ b/Engine/Source/Runtime/NetworkReplayStreaming/HttpNetworkReplayStreaming/Public/HttpNetworkReplayStreaming.h
@@ -103,7 +103,11 @@ class FHttpNetworkReplayStreamer;
 class FQueuedHttpRequest
 {
 public:
-	FQueuedHttpRequest( const EQueuedHttpRequestType::Type InType, TSharedPtr< class IHttpRequest > InRequest ) : Type( InType ), Request( InRequest )
+	FQueuedHttpRequest( const EQueuedHttpRequestType::Type InType, TSharedPtr< class IHttpRequest > InRequest ) : Type( InType ), Request( InRequest ), RetryProgress( 0 ), MaxRetries( 0 ), RetryDelay( 0.0f ), NextRetryTime( 0.0 )
+	{
+	}
+
+	FQueuedHttpRequest( const EQueuedHttpRequestType::Type InType, TSharedPtr< class IHttpRequest > InRequest, const int32 InMaxRetries, const float InRetryDelay ) : Type( InType ), Request( InRequest ), RetryProgress( 0 ), MaxRetries( InMaxRetries ), RetryDelay( InRetryDelay ), NextRetryTime( 0.0 )
 	{
 	}
 
@@ -114,6 +118,11 @@ public:
 	EQueuedHttpRequestType::Type		Type;
 	TSharedPtr< class IHttpRequest >	Request;
 
+	int32								RetryProgress;
+	int32								MaxRetries;
+	float								RetryDelay;
+	double								NextRetryTime;
+
 	virtual bool PreProcess( FHttpNetworkReplayStreamer* Streamer, const FString& ServerURL, const FString& SessionName )
 	{
 		return true;
@@ -226,8 +235,9 @@ public:
 	void FlushCheckpointInternal( uint32 TimeInMS );
 	virtual void AddEvent( const uint32 TimeInMS, const FString& Group, const FString& Meta, const TArray<uint8>& Data ) override;
 	virtual void AddOrUpdateEvent( const FString& Name, const uint32 TimeInMS, const FString& Group, const FString& Meta, const TArray<uint8>& Data ) override;
-	void AddRequestToQueue( const EQueuedHttpRequestType::Type Type, TSharedPtr< class IHttpRequest >	Request );
+	void AddRequestToQueue( const EQueuedHttpRequestType::Type Type, TSharedPtr< class IHttpRequest > Request, const int32 InMaxRetries = 0, const float InRetryDelay = 0.0f );
 	void AddCustomRequestToQueue( TSharedPtr< FQueuedHttpRequest > Request );
+	bool RetryRequest( TSharedPtr< FQueuedHttpRequest > Request, FHttpResponsePtr HttpResponse );
 	void EnumerateCheckpoints();
 	void ConditionallyEnumerateCheckpoints();
 
diff --git a/Engine/Source/Runtime/Online/Hotfix/Private/OnlineHotfixManager.cpp b/Engine/Source/Runtime/Online/Hotfix/Private/OnlineHotfixManager.cpp
--- a/Engine/Source/Runtime/Online/Hotfix/Private/OnlineHotfixManager.cpp
+++ b/Engine/Source/Runtime/Online/Hotfix/Private/OnlineHotfixManager.cpp
@@ -466,7 +466,7 @@ FString UOnlineHotfixManager::GetStrippedConfigFileName(const FString& IniName)
 	}
 	else if (StrippedIniName.StartsWith(DefaultPrefix))
 	{
-		StrippedIniName = IniName.Right(StrippedIniName.Len() - 7);
+		StrippedIniName = IniName.Right(StrippedIniName.Len() - DefaultPrefix.Len());
 	}
 	else if (StrippedIniName.StartsWith(DebugPrefix))
 	{
diff --git a/Engine/Source/Runtime/PerfCounters/Private/PerfCounters.cpp b/Engine/Source/Runtime/PerfCounters/Private/PerfCounters.cpp
--- a/Engine/Source/Runtime/PerfCounters/Private/PerfCounters.cpp
+++ b/Engine/Source/Runtime/PerfCounters/Private/PerfCounters.cpp
@@ -243,10 +243,18 @@ bool FPerfCounters::Tick(float DeltaTime)
 	float CurrentTime = FPlatformTime::Seconds();
 	if (CurrentTime - LastTimeInternalCountersUpdated > InternalCountersUpdateInterval)
 	{
+		// get CPU stats first
+		FCPUTime CPUStats = FPlatformTime::GetCPUTime();
+		Set(TEXT("ProcessCPUUsageRelativeToCore"), CPUStats.CPUTimePctRelative);
+
+		// memory
 		FPlatformMemoryStats Stats = FPlatformMemory::GetStats();
 		Set(TEXT("AvailablePhysicalMemoryMB"), static_cast<uint64>(Stats.AvailablePhysical / (1024 * 1024)));
 		Set(TEXT("AvailableVirtualMemoryMB"), static_cast<uint64>(Stats.AvailableVirtual / (1024 * 1024)));
+		Set(TEXT("ProcessPhysicalMemoryMB"), static_cast<uint64>(Stats.UsedPhysical/ (1024 * 1024)));
+		Set(TEXT("ProcessVirtualMemoryMB"), static_cast<uint64>(Stats.UsedVirtual / (1024 * 1024)));
 
+		// disk space
 		const FString LogFilename = FPlatformOutputDevices::GetAbsoluteLogFilename();
 		uint64 TotalBytesOnLogDrive = 0, FreeBytesOnLogDrive = 0;
 		if (FPlatformMisc::GetDiskTotalAndFreeSpace(LogFilename, TotalBytesOnLogDrive, FreeBytesOnLogDrive))
diff --git a/Engine/Source/Runtime/RHI/Private/DynamicRHI.cpp b/Engine/Source/Runtime/RHI/Private/DynamicRHI.cpp
--- a/Engine/Source/Runtime/RHI/Private/DynamicRHI.cpp
+++ b/Engine/Source/Runtime/RHI/Private/DynamicRHI.cpp
@@ -76,8 +76,8 @@ void RHIDetectAndWarnOfBadDrivers()
 	{
 		DriverInfo.SetNVIDIA();
 		DriverInfo.DeviceDescription = TEXT("Test NVIDIA (bad)");
-		DriverInfo.UserDriverVersion = TEXT("361.43");
-		DriverInfo.InternalDriverVersion = TEXT("9.18.136.143");
+		DriverInfo.UserDriverVersion = TEXT("346.43");
+		DriverInfo.InternalDriverVersion = TEXT("9.18.134.643");
 		DriverInfo.DriverDate = TEXT("01-01-1900");
 	}
 	else if(CVarValue == 3)
@@ -114,22 +114,44 @@ void RHIDetectAndWarnOfBadDrivers()
 
 		if (BlackListEntry.IsValid())
 		{
+			bool bLatestBlacklisted = DetectedGPUHardware.IsLatestBlacklisted();
+
+			// Note: we don't localize the vendor's name.
+			FString VendorString = DriverInfo.ProviderName;
+			if (DriverInfo.IsNVIDIA())
+			{
+				VendorString = TEXT("NVIDIA");
+			}
+			else if (DriverInfo.IsAMD())
+			{
+				VendorString = TEXT("AMD");
+			}
+			else if (DriverInfo.IsIntel())
+			{
+				VendorString = TEXT("Intel");
+			}
+
 			// format message box UI
 			FFormatNamedArguments Args;
 			Args.Add(TEXT("AdapterName"), FText::FromString(DriverInfo.DeviceDescription));
-			Args.Add(TEXT("UserDriverVersion"), FText::FromString(DriverInfo.UserDriverVersion));
-			Args.Add(TEXT("InternalDriverVersion"), FText::FromString(DriverInfo.InternalDriverVersion));
-			Args.Add(TEXT("DriverDate"), FText::FromString(DriverInfo.DriverDate));
-			Args.Add(TEXT("SuggestedVersion"), FText::FromString(DetectedGPUHardware.GetSuggestedDriverVersion()));
+			Args.Add(TEXT("Vendor"), FText::FromString(VendorString));
+			Args.Add(TEXT("RecommendedVer"), FText::FromString(DetectedGPUHardware.GetSuggestedDriverVersion()));
+			Args.Add(TEXT("InstalledVer"), FText::FromString(DriverInfo.UserDriverVersion));
 
 			// this message can be suppressed with r.WarnOfBadDrivers=0
-			FText LocalizedMsg = FText::Format(NSLOCTEXT("MessageDialog", "VideoCardDriverIssueReport",
-				"Your current video card driver version might cause problems.\n(reduced performance, artifacts, stalls, crashes, restarts)\n\nYour current driver:\n  Name:\t{AdapterName}\n  Version:\t{UserDriverVersion} (internal {InternalDriverVersion})\n  Date:\t{DriverDate}\n\nSuggested driver:\n  Version:\t{SuggestedVersion}"),
-				Args);
+			FText LocalizedMsg;
+			if (bLatestBlacklisted)
+			{
+				LocalizedMsg = FText::Format(NSLOCTEXT("MessageDialog", "LatestVideoCardDriverIssueReport","The latest version of the {Vendor} graphics driver has known issues.\n\nPlease install the last known good driver version.\n\n{AdapterName}\n{RecommendedVer} is the last known good\n{InstalledVer} is installed"),Args);
+			}
+			else
+			{
+				LocalizedMsg = FText::Format(NSLOCTEXT("MessageDialog", "VideoCardDriverIssueReport","Your {Vendor} graphics driver has known issues.\n\nPlease update to the latest driver version.\n\n{AdapterName}\n{RecommendedVer} is recommended\n{InstalledVer} is installed"),Args);
+			}
 
 			FPlatformMisc::MessageBoxExt(EAppMsgType::Ok,
 				*LocalizedMsg.ToString(),
-				*NSLOCTEXT("MessageDialog", "TitleVideoCardDriverIssue", "WARNING: Video card driver").ToString());
+				*NSLOCTEXT("MessageDialog", "TitleVideoCardDriverIssue", "WARNING: Known issues with graphics driver").ToString());
 		}
 	}
 }
diff --git a/Engine/Source/Runtime/RHI/Private/GPUDefragAllocator.cpp b/Engine/Source/Runtime/RHI/Private/GPUDefragAllocator.cpp
--- a/Engine/Source/Runtime/RHI/Private/GPUDefragAllocator.cpp
+++ b/Engine/Source/Runtime/RHI/Private/GPUDefragAllocator.cpp
@@ -5,7 +5,7 @@ Copyright 1998-2015 Epic Games, Inc. All Rights Reserved.
 
 #include "RHI.h"
 #include "GPUDefragAllocator.h"
-
+#include "ScopedTimers.h"
 
 DECLARE_STATS_GROUP(TEXT("TexturePool"), STATGROUP_TexturePool, STATCAT_ADVANCED);
 
@@ -1139,6 +1139,23 @@ void FGPUDefragAllocator::InsertFence()
 	CurrentSyncIndex++;
 }
 
+#if VALIDATE_MEMORY_PROTECTION
+void FGPUDefragAllocator::AllowCPUAccessToBlocks()
+{
+	FScopedDurationTimer Timer(TimeInMemProtect);
+	const int32 Count = BlocksToUnProtect.Num();
+	for (int32 i = Count - 1; i >= 0; --i)
+	{
+		const FMemProtectTracker& Block = BlocksToUnProtect[i];
+		if (Block.SyncIndex >= CompletedSyncIndex)
+		{
+			PlatformAllowCPUAccessToBlock(Block);
+			BlocksToUnProtect.RemoveAtSwap(i, 1, false);
+		}
+	}	
+}
+#endif
+
 /**
 * Blocks the calling thread until the current sync fence has been completed.
 */
@@ -1154,6 +1171,10 @@ void FGPUDefragAllocator::BlockOnFence()
 		}
 		CompletedSyncIndex = CurrentSyncIndex - 1;
 		BlockedCycles += FPlatformTime::Cycles() - StartTime;
+
+#if VALIDATE_MEMORY_PROTECTION
+		AllowCPUAccessToBlocks();
+#endif
 	}
 }
 
@@ -1438,6 +1459,11 @@ void FGPUDefragAllocator::FullDefragmentation(FRelocationStats& Stats)
 int32 FGPUDefragAllocator::Tick(FRelocationStats& Stats, bool bPanicDefrag)
 {
 	FScopeLock SyncLock(&SynchronizationObject);
+
+#if VALIDATE_MEMORY_PROTECTION
+	TimeInMemProtect = 0;
+#endif
+
 	SET_CYCLE_COUNTER(STAT_TexturePool_Blocked, BlockedCycles);
 	double StartTime = FPlatformTime::Seconds();
 	BlockedCycles = 0;
@@ -1474,6 +1500,14 @@ int32 FGPUDefragAllocator::Tick(FRelocationStats& Stats, bool bPanicDefrag)
 		}
 	}
 
+#if VALIDATE_MEMORY_PROTECTION
+	{
+		FScopedDurationTimer MemProtectTimer(TimeInMemProtect);
+		PlatformRemoveCPUAccessTo(BlocksToProtect);
+		BlocksToUnProtect.Append(BlocksToProtect);
+		BlocksToProtect.Reset();
+	}
+#endif
 	NumRelocationsInProgress = Stats.NumRelocations;
 
 	// Start a new sync point.
@@ -1490,7 +1524,7 @@ int32 FGPUDefragAllocator::Tick(FRelocationStats& Stats, bool bPanicDefrag)
 	FMemoryChunk* TestChunk = FirstChunk;
 	while (TestChunk)
 	{		
-		FMemoryChunk* TestAgainstChunk = FirstChunk->NextChunk;
+		FMemoryChunk* TestAgainstChunk = TestChunk->NextChunk;
 		while (TestAgainstChunk)
 		{			
 			bool bBefore = (TestChunk->Base + TestChunk->Size) <= TestAgainstChunk->Base;
