Copying //UE4/Dev-Networking to //UE4/Dev-Main (Source: //UE4/Dev-Networking @ 3802510)

diff --git a/Engine/Source/Runtime/NetworkReplayStreaming/LocalFileNetworkReplayStreaming/Private/LocalFileNetworkReplayStreaming.cpp b/Engine/Source/Runtime/NetworkReplayStreaming/LocalFileNetworkReplayStreaming/Private/LocalFileNetworkReplayStreaming.cpp
new file mode 100644
--- /dev/null
+++ b/Engine/Source/Runtime/NetworkReplayStreaming/LocalFileNetworkReplayStreaming/Private/LocalFileNetworkReplayStreaming.cpp
@@ -0,0 +1,996 @@
+// Copyright 1998-2017 Epic Games, Inc. All Rights Reserved.
+
+#include "LocalFileNetworkReplayStreaming.h"
+#include "Misc/NetworkVersion.h"
+#include "HAL/FileManager.h"
+#include "Misc/Paths.h"
+#include "Misc/Guid.h"
+
+DEFINE_LOG_CATEGORY_STATIC(LogLocalFileReplay, Log, All);
+
+namespace LocalFileReplay
+{
+	const uint32 FileMagic = 0x1CA2E27F;
+	const uint32 FileVersion = 0;
+
+	FString GetDemoPath()
+	{
+		return FPaths::Combine(*FPaths::ProjectSavedDir(), TEXT("Demos/"));
+	}
+
+	FString GetDemoFullFilename(const FString& StreamName)
+	{
+		return FPaths::Combine(*GetDemoPath(), *StreamName) + TEXT(".replay");
+	}
+
+	// Returns a name formatted as "demoX", where X is 0-9.
+	// Returns the first value that doesn't yet exist, or if they all exist, returns the oldest one
+	// (it will be overwritten).
+	FString GetAutomaticDemoName()
+	{
+		FString FinalDemoName;
+		FDateTime BestDateTime = FDateTime::MaxValue();
+
+		const int MAX_DEMOS = 10;
+
+		for (int32 i = 0; i < MAX_DEMOS; i++)
+		{
+			const FString DemoName = FString::Printf(TEXT("demo%i"), i + 1);
+			const FString FullDemoName = LocalFileReplay::GetDemoFullFilename(DemoName);
+
+			FDateTime DateTime = IFileManager::Get().GetTimeStamp(*FullDemoName);
+			if (DateTime == FDateTime::MinValue())
+			{
+				// If we don't find this file, we can early out now
+				FinalDemoName = DemoName;
+				break;
+			}
+			else if (DateTime < BestDateTime)
+			{
+				// Use the oldest file
+				FinalDemoName = DemoName;
+				BestDateTime = DateTime;
+			}
+		}
+
+		return FinalDemoName;
+	}
+};
+
+void FLocalFileStreamFArchive::SetChunk(int32 ChunkIndex)
+{
+	check(Streamer != nullptr);
+	check(Streamer->ReplayInfo.DataChunks.IsValidIndex(ChunkIndex));
+	check(LocalFileAr.Get());
+
+	CurrentChunkIndex = ChunkIndex;
+
+	FLocalFileChunkInfo ChunkInfo = Streamer->ReplayInfo.Chunks[Streamer->ReplayInfo.DataChunks[CurrentChunkIndex].ChunkIndex];
+
+	LocalFileAr->Seek(ChunkInfo.Offset);
+
+	CurrentChunkBuffer.Reset(ChunkInfo.SizeInBytes);
+	CurrentChunkBuffer.AddUninitialized(ChunkInfo.SizeInBytes);
+
+	LocalFileAr->Serialize(CurrentChunkBuffer.GetData(), CurrentChunkBuffer.Num());
+
+	if (Streamer->SupportsCompression())
+	{
+		TArray<uint8> Uncompressed;
+		if (Streamer->DecompressBuffer(CurrentChunkBuffer, Uncompressed))
+		{
+			CurrentChunkBuffer = Uncompressed;
+		}
+		else
+		{
+			ArIsError = true;
+			return;
+		}
+	}
+
+	Pos = Streamer->ReplayInfo.DataChunks[CurrentChunkIndex].ReplayOffset;
+	CurrentChunkPos = 0;
+}
+
+void FLocalFileStreamFArchive::Serialize(void* V, int64 Length)
+{
+	if (IsLoading())
+	{
+		if (Pos + Length > TotalSize())
+		{
+			ArIsError = true;
+			return;
+		}
+
+		check((CurrentChunkPos + Length) <= CurrentChunkBuffer.Num());
+
+		FMemory::Memcpy(V, CurrentChunkBuffer.GetData() + CurrentChunkPos, Length);
+
+		Pos += Length;
+		CurrentChunkPos += Length;
+
+		if ((CurrentChunkPos >= CurrentChunkBuffer.Num()) && !AtEnd())
+		{
+			CurrentChunkIndex++;
+			SetChunk(CurrentChunkIndex);
+		}
+	}
+	else
+	{
+		ArIsError = true;
+	}
+}
+
+int64 FLocalFileStreamFArchive::Tell()
+{
+	return Pos;
+}
+
+int64 FLocalFileStreamFArchive::TotalSize()
+{
+	check(Streamer != nullptr);
+	return Streamer->ReplayInfo.TotalDataSizeInBytes;
+}
+
+void FLocalFileStreamFArchive::Seek(int64 InPos)
+{
+	check(Streamer != nullptr);
+	check(InPos < TotalSize());
+
+	int32 ChunkIndex = 0;
+
+	// figure out current chunk, set it, then offset within it
+	for(int32 i=0; i < Streamer->ReplayInfo.DataChunks.Num(); ++i)
+	{
+		if (Streamer->ReplayInfo.DataChunks[i].ReplayOffset <= InPos)
+		{
+			ChunkIndex = i;
+		}
+		else
+		{
+			break;
+		}
+	}
+	
+	SetChunk(ChunkIndex);
+
+	Pos = InPos;
+	CurrentChunkPos = (Pos - Streamer->ReplayInfo.DataChunks[CurrentChunkIndex].ReplayOffset);
+}
+
+bool FLocalFileStreamFArchive::AtEnd()
+{
+	return Pos >= TotalSize();
+}
+
+void FLocalFileHeaderFArchive::Flush()
+{
+	FArrayWriter::Flush();
+
+	check(Streamer != nullptr);
+	Streamer->FlushHeader();
+}
+
+FLocalFileReplayInfo FLocalFileNetworkReplayStreamer::ReadReplayInfo(const FString& StreamName) const
+{
+	const FString FullFilename = LocalFileReplay::GetDemoFullFilename(StreamName);
+
+	FLocalFileReplayInfo Info;
+	TUniquePtr<FArchive> LocalFileAr(IFileManager::Get().CreateFileReader(*FullFilename, FILEREAD_AllowWrite));
+
+	if (LocalFileAr.IsValid() && LocalFileAr->TotalSize() != 0)
+	{
+		uint32 MagicNumber;
+		*LocalFileAr << MagicNumber;
+
+		uint32 FileVersion;
+		*LocalFileAr << FileVersion;
+
+		if ((MagicNumber == LocalFileReplay::FileMagic) && (FileVersion == LocalFileReplay::FileVersion))
+		{
+			// read summary info
+			*LocalFileAr << Info.LengthInMS;
+			*LocalFileAr << Info.NetworkVersion;
+			*LocalFileAr << Info.Changelist;
+			*LocalFileAr << Info.FriendlyName;
+
+			uint32 IsLive;
+			*LocalFileAr << IsLive;
+
+			Info.bIsLive = (IsLive != 0);
+
+			// now look for all chunks
+			while (!LocalFileAr->AtEnd())
+			{
+				ELocalFileChunkType ChunkType;
+				*LocalFileAr << ChunkType;
+
+				int32 Idx = Info.Chunks.AddDefaulted();
+				Info.Chunks[Idx].ChunkType = ChunkType;
+
+				*LocalFileAr << Info.Chunks[Idx].SizeInBytes;
+
+				Info.Chunks[Idx].Offset = LocalFileAr->Tell();
+
+				switch(ChunkType)
+				{
+				case ELocalFileChunkType::Header:
+					check(Info.HeaderChunkIndex == INDEX_NONE);
+					Info.HeaderChunkIndex = Idx;
+					break;
+				case ELocalFileChunkType::Checkpoint:
+				{
+					int32 CheckpointIdx = Info.Checkpoints.AddDefaulted();
+					Info.Checkpoints[CheckpointIdx].ChunkIndex = Idx;
+
+					*LocalFileAr << Info.Checkpoints[CheckpointIdx].Id;
+					*LocalFileAr << Info.Checkpoints[CheckpointIdx].Group;
+					*LocalFileAr << Info.Checkpoints[CheckpointIdx].Metadata;
+					*LocalFileAr << Info.Checkpoints[CheckpointIdx].Time1;
+					*LocalFileAr << Info.Checkpoints[CheckpointIdx].Time2;
+
+					*LocalFileAr << Info.Checkpoints[CheckpointIdx].SizeInBytes;
+
+					Info.Checkpoints[CheckpointIdx].Offset = LocalFileAr->Tell();
+				}
+				break;
+				case ELocalFileChunkType::ReplayData:
+				{
+					int32 DataIdx = Info.DataChunks.AddDefaulted();
+					Info.DataChunks[DataIdx].ChunkIndex = Idx;
+					Info.DataChunks[DataIdx].ReplayOffset = Info.TotalDataSizeInBytes;
+
+					if (SupportsCompression())
+					{
+						int64 SavedPos = LocalFileAr->Tell();
+						LocalFileAr->Seek(Info.Chunks[Idx].Offset);
+						Info.TotalDataSizeInBytes += GetDecompressedSize(*LocalFileAr);
+						LocalFileAr->Seek(SavedPos);
+					}
+					else
+					{
+						Info.TotalDataSizeInBytes += Info.Chunks[Idx].SizeInBytes;
+					}
+				}
+				break;
+				case ELocalFileChunkType::Event:
+				{
+					int32 EventIdx = Info.Events.AddDefaulted();
+					Info.Events[EventIdx].ChunkIndex = Idx;
+
+					*LocalFileAr << Info.Events[EventIdx].Id;
+					*LocalFileAr << Info.Events[EventIdx].Group;
+					*LocalFileAr << Info.Events[EventIdx].Metadata;
+					*LocalFileAr << Info.Events[EventIdx].Time1;
+					*LocalFileAr << Info.Events[EventIdx].Time2;
+
+					*LocalFileAr << Info.Events[EventIdx].SizeInBytes;
+
+					Info.Events[EventIdx].Offset = LocalFileAr->Tell();
+				}
+				break;
+				default:
+					UE_LOG(LogLocalFileReplay, Log, TEXT("ReadReplayInfo: Unhandled file chunk type: %d"), (uint32)ChunkType);
+					break;
+				}
+
+				LocalFileAr->Seek(Info.Chunks[Idx].Offset + Info.Chunks[Idx].SizeInBytes);
+			}
+		}
+
+		Info.bIsValid = (Info.HeaderChunkIndex != INDEX_NONE) && (Info.TotalDataSizeInBytes > 0);
+	}
+
+	return Info;
+}
+
+void FLocalFileNetworkReplayStreamer::WriteReplayInfo(const FString& StreamName, FLocalFileReplayInfo& InReplayInfo)
+{
+	// Update metadata with latest info
+	TUniquePtr<FArchive> ReplayInfoFileAr(IFileManager::Get().CreateFileWriter(*LocalFileReplay::GetDemoFullFilename(StreamName), FILEWRITE_Append | FILEWRITE_AllowRead));
+	if (ReplayInfoFileAr.IsValid())
+	{
+		ReplayInfoFileAr->Seek(0);
+
+		uint32 FileMagic = LocalFileReplay::FileMagic;
+		*ReplayInfoFileAr << FileMagic;
+		
+		uint32 FileVersion = LocalFileReplay::FileVersion;
+		*ReplayInfoFileAr << FileVersion;
+
+		*ReplayInfoFileAr << InReplayInfo.LengthInMS;
+		*ReplayInfoFileAr << InReplayInfo.NetworkVersion;
+		*ReplayInfoFileAr << InReplayInfo.Changelist;
+		*ReplayInfoFileAr << InReplayInfo.FriendlyName;
+
+		uint32 IsLive = InReplayInfo.bIsLive ? 1 : 0;
+		*ReplayInfoFileAr << IsLive;
+	}
+}
+
+void FLocalFileNetworkReplayStreamer::StartStreaming(const FString& CustomName, const FString& FriendlyName, const TArray< FString >& UserNames, bool bRecord, const FNetworkReplayVersion& ReplayVersion, const FOnStreamReadyDelegate& Delegate)
+{
+	FString FinalDemoName = CustomName;
+
+	if (CustomName.IsEmpty())
+	{
+		if (bRecord)
+		{
+			// If we're recording and the caller didn't provide a name, generate one automatically
+			FinalDemoName = LocalFileReplay::GetAutomaticDemoName();
+		}
+		else
+		{
+			// Can't play a replay if the user didn't provide a name!
+			Delegate.ExecuteIfBound(false, bRecord);
+			return;
+		}
+	}
+
+	const FString FullDemoFilename = LocalFileReplay::GetDemoFullFilename(FinalDemoName);
+
+	CurrentStreamName = FinalDemoName;
+
+	if (!bRecord)
+	{
+		// Load metadata if it exists
+		ReplayInfo = ReadReplayInfo(CurrentStreamName);
+
+		if (ReplayInfo.bIsValid)
+		{
+			HeaderPlaybackAr.Reset(IFileManager::Get().CreateFileReader(*FullDemoFilename, FILEREAD_AllowWrite));
+			HeaderPlaybackAr->Seek(ReplayInfo.Chunks[ReplayInfo.HeaderChunkIndex].Offset);
+
+			if (ReplayInfo.DataChunks.Num() > 0)
+			{
+				FilePlaybackAr.Reset(new FLocalFileStreamFArchive(this));
+				FilePlaybackAr->ArIsLoading = 1;
+				FilePlaybackAr->LocalFileAr.Reset(IFileManager::Get().CreateFileReader(*FullDemoFilename, FILEREAD_AllowWrite));
+				FilePlaybackAr->Seek(0);
+			}
+
+			StreamerState = EStreamerState::Playback;
+		}
+	}
+	else
+	{
+		// Delete any existing demo with this name
+		IFileManager::Get().Delete(*FullDemoFilename);
+
+		FileRecordingAr.Reset(new FArrayWriter());
+	
+		StreamerState = EStreamerState::Recording;
+
+		// Set up replay info
+		ReplayInfo.NetworkVersion = ReplayVersion.NetworkVersion;
+		ReplayInfo.Changelist = ReplayVersion.Changelist;
+		ReplayInfo.FriendlyName = FriendlyName;
+		ReplayInfo.bIsLive = true;
+
+		WriteReplayInfo(CurrentStreamName, ReplayInfo);
+
+		HeaderRecordingAr.Reset(new FLocalFileHeaderFArchive(this));
+	}
+
+	// Notify immediately
+	Delegate.ExecuteIfBound(GetStreamingArchive() != nullptr && GetHeaderArchive() != nullptr, bRecord);
+}
+
+void FLocalFileNetworkReplayStreamer::FlushHeader()
+{
+	if (HeaderRecordingAr.Get() && (HeaderRecordingAr->Num() > 0))
+	{
+		TUniquePtr<FArchive> LocalFileAr(IFileManager::Get().CreateFileWriter(*LocalFileReplay::GetDemoFullFilename(CurrentStreamName), FILEWRITE_Append | FILEWRITE_AllowRead));
+		if (LocalFileAr.IsValid())
+		{
+			LocalFileAr->Seek(LocalFileAr->TotalSize());
+
+			ELocalFileChunkType ChunkType = ELocalFileChunkType::Header;
+			*LocalFileAr << ChunkType;
+
+			int32 ChunkSize = HeaderRecordingAr->Num();
+			*LocalFileAr << ChunkSize;
+
+			LocalFileAr->Serialize(HeaderRecordingAr->GetData(), HeaderRecordingAr->Num());
+		}
+
+		HeaderRecordingAr.Reset();
+	}
+}
+
+void FLocalFileNetworkReplayStreamer::StopStreaming()
+{
+	if (StreamerState == EStreamerState::Recording)
+	{
+		FlushStream();
+		FlushHeader();
+
+		ReplayInfo.bIsLive = false;
+
+		WriteReplayInfo(CurrentStreamName, ReplayInfo);
+	}
+	
+	FilePlaybackAr.Reset();
+	FileRecordingAr.Reset();
+	HeaderRecordingAr.Reset();
+	HeaderPlaybackAr.Reset();
+	CheckpointPlaybackAr.Reset();
+	CheckpointRecordingAr.Reset();
+
+	CurrentStreamName.Empty();
+	StreamerState = EStreamerState::Idle;
+}
+
+FArchive* FLocalFileNetworkReplayStreamer::GetHeaderArchive()
+{
+	if (StreamerState == EStreamerState::Recording)
+	{
+		return HeaderRecordingAr.Get();
+	}
+
+	return HeaderPlaybackAr.Get();
+}
+
+FArchive* FLocalFileNetworkReplayStreamer::GetStreamingArchive()
+{
+	if (StreamerState == EStreamerState::Recording)
+	{
+		return FileRecordingAr.Get();
+	}
+
+	return FilePlaybackAr.Get();
+}
+
+void FLocalFileNetworkReplayStreamer::UpdateTotalDemoTime(uint32 TimeInMS)
+{
+	check(StreamerState == EStreamerState::Recording);
+
+	ReplayInfo.LengthInMS = TimeInMS;
+}
+
+bool FLocalFileNetworkReplayStreamer::IsDataAvailable() const
+{
+	check(StreamerState == EStreamerState::Playback);
+
+	return FilePlaybackAr.IsValid() && FilePlaybackAr->Tell() < FilePlaybackAr->TotalSize();
+}
+
+bool FLocalFileNetworkReplayStreamer::IsLive() const
+{
+	return IsNamedStreamLive(CurrentStreamName);
+}
+
+bool FLocalFileNetworkReplayStreamer::IsNamedStreamLive(const FString& StreamName) const
+{
+	FLocalFileReplayInfo Info = ReadReplayInfo(StreamName);
+	return Info.bIsLive;
+}
+
+void FLocalFileNetworkReplayStreamer::DeleteFinishedStream(const FString& StreamName, const FOnDeleteFinishedStreamComplete& Delegate) const
+{
+	// Live streams can't be deleted
+	if (IsNamedStreamLive(StreamName))
+	{
+		UE_LOG(LogLocalFileReplay, Log, TEXT("Can't delete network replay stream %s because it is live!"), *StreamName);
+		Delegate.ExecuteIfBound(false);
+		return;
+	}
+
+	// Delete the directory with the specified name in the Saved/Demos directory
+	const FString FullDemoFilename = LocalFileReplay::GetDemoFullFilename(StreamName);
+
+	const bool DeleteSucceeded = IFileManager::Get().Delete(*FullDemoFilename);
+
+	Delegate.ExecuteIfBound(DeleteSucceeded);
+}
+
+void FLocalFileNetworkReplayStreamer::EnumerateStreams(const FNetworkReplayVersion& ReplayVersion, const FString& UserString, const FString& MetaString, const FOnEnumerateStreamsComplete& Delegate)
+{
+	EnumerateStreams(ReplayVersion, UserString, MetaString, TArray< FString >(), Delegate);
+}
+
+void FLocalFileNetworkReplayStreamer::EnumerateStreams(const FNetworkReplayVersion& ReplayVersion, const FString& UserString, const FString& MetaString, const TArray< FString >& ExtraParms, const FOnEnumerateStreamsComplete& Delegate)
+{
+	const FString WildCardPath = LocalFileReplay::GetDemoPath() + TEXT("*.replay");
+
+	TArray<FString> ReplayFileNames;
+	IFileManager::Get().FindFiles(ReplayFileNames, *WildCardPath, true, false);
+
+	TArray<FNetworkReplayStreamInfo> Results;
+
+	for (const FString& ReplayFileName : ReplayFileNames)
+	{
+		FNetworkReplayStreamInfo Info;
+
+		// Read stored info for this replay
+		FLocalFileReplayInfo StoredReplayInfo = ReadReplayInfo(FPaths::GetBaseFilename(ReplayFileName));
+		if (!StoredReplayInfo.bIsValid)
+		{
+			continue;
+		}
+
+		// Check version. NetworkVersion and changelist of 0 will ignore version check.
+		const bool NetworkVersionMatches = ReplayVersion.NetworkVersion == StoredReplayInfo.NetworkVersion;
+		const bool ChangelistMatches = ReplayVersion.Changelist == StoredReplayInfo.Changelist;
+
+		const bool NetworkVersionPasses = ReplayVersion.NetworkVersion == 0 || NetworkVersionMatches;
+		const bool ChangelistPasses = ReplayVersion.Changelist == 0 || ChangelistMatches;
+
+		if (NetworkVersionPasses && ChangelistPasses)
+		{
+			Info.Name = FPaths::GetBaseFilename(ReplayFileName);
+			Info.Timestamp = IFileManager::Get().GetTimeStamp(*LocalFileReplay::GetDemoFullFilename(Info.Name));
+			Info.bIsLive = StoredReplayInfo.bIsLive;
+			Info.LengthInMS = StoredReplayInfo.LengthInMS;
+			Info.FriendlyName = StoredReplayInfo.FriendlyName;
+			Info.SizeInBytes = StoredReplayInfo.TotalDataSizeInBytes;
+
+			Results.Add(Info);
+		}
+	}
+
+	Delegate.ExecuteIfBound(Results);
+}
+
+void FLocalFileNetworkReplayStreamer::AddUserToReplay(const FString& UserString)
+{
+	UE_LOG(LogLocalFileReplay, Log, TEXT("FLocalFileNetworkReplayStreamer::AddUserToReplay is currently unsupported."));
+}
+
+void FLocalFileNetworkReplayStreamer::AddEvent(const uint32 TimeInMS, const FString& Group, const FString& Meta, const TArray<uint8>& Data)
+{
+	if (StreamerState != EStreamerState::Recording)
+	{
+		UE_LOG(LogLocalFileReplay, Warning, TEXT("FLocalFileNetworkReplayStreamer::AddEvent. Not recording."));
+		return;
+	}
+
+	AddOrUpdateEvent(TEXT(""), TimeInMS, Group, Meta, Data);
+}
+
+void FLocalFileNetworkReplayStreamer::AddOrUpdateEvent(const FString& Name, const uint32 TimeInMS, const FString& Group, const FString& Meta, const TArray<uint8>& Data)
+{
+	if (StreamerState != EStreamerState::Recording)
+	{
+		UE_LOG(LogLocalFileReplay, Warning, TEXT("FLocalFileNetworkReplayStreamer::AddOrUpdateEvent. Not recording."));
+		return;
+	}
+
+	UE_LOG(LogLocalFileReplay, Verbose, TEXT("FLocalFileNetworkReplayStreamer::AddOrUpdateEvent. Size: %i"), Data.Num());
+
+	int32 EventIndex = INDEX_NONE;
+
+	FString EventName = Name;
+
+	// if name is empty, assign one
+	if (EventName.IsEmpty())
+	{
+		EventName = FGuid::NewGuid().ToString(EGuidFormats::Digits);
+	}
+
+	// prefix with stream name to be consistent with http streamer
+	EventName = CurrentStreamName + TEXT("_") + EventName;
+
+	// see if this event already exists
+	if (EventIndex == INDEX_NONE)
+	{
+		for (int32 i=0; i < ReplayInfo.Events.Num(); ++i)
+		{
+			if (ReplayInfo.Events[i].Id == EventName)
+			{
+				EventIndex = i;
+				break;
+			}
+		}
+	}
+	
+	TUniquePtr<FArchive> LocalFileAr(IFileManager::Get().CreateFileWriter(*LocalFileReplay::GetDemoFullFilename(CurrentStreamName), FILEWRITE_Append | FILEWRITE_AllowRead));
+	if (LocalFileAr.IsValid())
+	{
+		if (EventIndex == INDEX_NONE)
+		{
+			// append new event chunk
+			LocalFileAr->Seek(LocalFileAr->TotalSize());
+
+			int32 Idx = ReplayInfo.Events.AddDefaulted();
+			ReplayInfo.Events[Idx].Id = EventName;
+			ReplayInfo.Events[Idx].Offset = LocalFileAr->Tell();
+			ReplayInfo.Events[Idx].SizeInBytes = Data.Num();
+		}
+		else 
+		{
+			LocalFileAr->Seek(ReplayInfo.Events[EventIndex].Offset);
+
+			// write data into existing space if possible
+			if (Data.Num() > ReplayInfo.Events[EventIndex].SizeInBytes)
+			{
+				// otherwise clear chunk type so it will be skipped later
+				ELocalFileChunkType ChunkType = ELocalFileChunkType::Unknown;
+				*LocalFileAr << ChunkType;
+
+				LocalFileAr->Seek(LocalFileAr->TotalSize());
+			}
+		}
+
+		// flush event
+		ELocalFileChunkType ChunkType = ELocalFileChunkType::Event;
+		*LocalFileAr << ChunkType;
+
+		int64 SavedPos = LocalFileAr->Tell();
+
+		int32 PlaceholderSize = 0;
+		*LocalFileAr << PlaceholderSize;
+
+		int64 MetadataPos = LocalFileAr->Tell();
+
+		*LocalFileAr << EventName;
+
+		FString Value = Group;
+		*LocalFileAr << Value;
+
+		Value = Meta;
+		*LocalFileAr << Value;
+
+		uint32 Time1 = TimeInMS;
+		*LocalFileAr << Time1;
+
+		uint32 Time2 = TimeInMS;
+		*LocalFileAr << Time2;
+
+		int32 EventSize = Data.Num();
+		*LocalFileAr << EventSize;
+
+		LocalFileAr->Serialize((void*)Data.GetData(), Data.Num());
+
+		int32 ChunkSize = LocalFileAr->Tell() - MetadataPos;
+
+		LocalFileAr->Seek(SavedPos);
+		*LocalFileAr << ChunkSize;
+	}
+}
+
+void FLocalFileNetworkReplayStreamer::EnumerateEvents(const FString& Group, const FEnumerateEventsCompleteDelegate& EnumerationCompleteDelegate)
+{
+	EnumerateEvents(CurrentStreamName, Group, EnumerationCompleteDelegate);
+}
+
+void FLocalFileNetworkReplayStreamer::EnumerateEvents(const FString& ReplayName, const FString& Group, const FEnumerateEventsCompleteDelegate& EnumerationCompleteDelegate)
+{
+	// Read stored info for this replay
+	FLocalFileReplayInfo StoredReplayInfo = ReadReplayInfo(ReplayName);
+	if (!StoredReplayInfo.bIsValid)
+	{
+		EnumerationCompleteDelegate.ExecuteIfBound(FReplayEventList(), false);
+		return;
+	}
+
+	FReplayEventList EventList;
+
+	for (const FLocalFileEventInfo& EventInfo : StoredReplayInfo.Events)
+	{
+		int Idx = EventList.ReplayEvents.AddDefaulted();
+
+		FReplayEventListItem& Event = EventList.ReplayEvents[Idx];
+		Event.ID = EventInfo.Id;
+		Event.Group = EventInfo.Group;
+		Event.Metadata = EventInfo.Metadata;
+		Event.Time1 = EventInfo.Time1;
+		Event.Time2 = EventInfo.Time2;
+	}
+
+	EnumerationCompleteDelegate.ExecuteIfBound(EventList, true);
+}
+
+void FLocalFileNetworkReplayStreamer::RequestEventData(const FString& EventID, const FOnRequestEventDataComplete& RequestEventDataComplete)
+{
+	// Assume current stream
+	FString StreamName = CurrentStreamName;
+
+	// But look for name prefix, http streamer expects to pull details from arbitrary streams
+	int32 Idx = INDEX_NONE;
+	if (EventID.FindChar(TEXT('_'), Idx))
+	{
+		StreamName = EventID.Left(Idx);
+	}
+
+	// Read stored info for this replay
+	FLocalFileReplayInfo StoredReplayInfo = ReadReplayInfo(StreamName);
+	if (StoredReplayInfo.bIsValid)
+	{
+		for (const FLocalFileEventInfo& EventInfo : StoredReplayInfo.Events)
+		{
+			if (EventInfo.Id == EventID)
+			{
+				TUniquePtr<FArchive> LocalFileAr(IFileManager::Get().CreateFileReader(*LocalFileReplay::GetDemoFullFilename(StreamName), FILEREAD_AllowWrite));
+				LocalFileAr->Seek(EventInfo.Offset);
+
+				TArray<uint8> Buffer;
+				Buffer.AddUninitialized(EventInfo.SizeInBytes);
+
+				LocalFileAr->Serialize(Buffer.GetData(), Buffer.Num());
+
+				RequestEventDataComplete.ExecuteIfBound(Buffer, true);
+				return;
+			}
+		}
+	}
+
+	RequestEventDataComplete.ExecuteIfBound(TArray<uint8>(), false);
+}
+
+void FLocalFileNetworkReplayStreamer::SearchEvents(const FString& EventGroup, const FOnEnumerateStreamsComplete& Delegate)
+{
+	UE_LOG(LogLocalFileReplay, Log, TEXT("FLocalFileNetworkReplayStreamer::SearchEvents is currently unsupported."));
+}
+
+FArchive* FLocalFileNetworkReplayStreamer::GetCheckpointArchive()
+{
+	if (StreamerState == EStreamerState::Recording)
+	{
+		// If the archive is null, and the API is being used properly, the caller is writing a checkpoint...
+		if (CheckpointRecordingAr.Get() == nullptr)
+		{
+			// Create a file writer for the next checkpoint index.
+			check(StreamerState != EStreamerState::Playback);
+
+			UE_LOG(LogLocalFileReplay, Log, TEXT("FLocalFileNetworkReplayStreamer::GetCheckpointArchive. Creating new checkpoint."));
+
+			// create in memory checkpoint, since it could take multiple frames to accumulate the checkpoint data
+			CheckpointRecordingAr.Reset(new FArrayWriter());
+		}
+
+		return CheckpointRecordingAr.Get();
+	}
+
+	check(StreamerState == EStreamerState::Playback);
+	return CheckpointPlaybackAr.Get();
+}
+
+void FLocalFileNetworkReplayStreamer::FlushStream()
+{
+	check(StreamerState == EStreamerState::Recording);
+
+	TUniquePtr<FArchive> LocalFileAr(IFileManager::Get().CreateFileWriter(*LocalFileReplay::GetDemoFullFilename(CurrentStreamName), FILEWRITE_Append | FILEWRITE_AllowRead));
+	if (LocalFileAr.IsValid())
+	{
+		LocalFileAr->Seek(LocalFileAr->TotalSize());
+
+		// flush chunk to disk
+		if (FileRecordingAr.Get() && (FileRecordingAr->Num() > 0))
+		{
+			// keep track of which replay chunk this is
+			ReplayInfo.DataChunks.AddDefaulted();
+
+			ELocalFileChunkType ChunkType = ELocalFileChunkType::ReplayData;
+			*LocalFileAr << ChunkType;
+
+			if (SupportsCompression())
+			{
+				TArray<uint8> Compressed;
+				if (CompressBuffer(*FileRecordingAr, Compressed))
+				{
+					int32 ChunkSize = Compressed.Num();
+					*LocalFileAr << ChunkSize;
+
+					LocalFileAr->Serialize(Compressed.GetData(), Compressed.Num());
+				}
+				else
+				{
+					UE_LOG(LogLocalFileReplay, Warning, TEXT("FLocalFileNetworkReplayStreamer::FlushStream - CompressBuffer failed"));
+				}
+			}
+			else
+			{
+				int32 ChunkSize = FileRecordingAr->Num();
+				*LocalFileAr << ChunkSize;
+
+				LocalFileAr->Serialize(FileRecordingAr->GetData(), FileRecordingAr->Num());
+			}
+
+			FileRecordingAr.Reset(new FArrayWriter);
+		}
+	}
+}
+
+void FLocalFileNetworkReplayStreamer::FlushCheckpoint(const uint32 TimeInMS)
+{
+	UE_LOG(LogLocalFileReplay, Log, TEXT("FLocalFileNetworkReplayStreamer::FlushCheckpoint. TimeInMS: %u"), TimeInMS);
+
+	FlushStream();
+
+	{
+		TUniquePtr<FArchive> LocalFileAr(IFileManager::Get().CreateFileWriter(*LocalFileReplay::GetDemoFullFilename(CurrentStreamName), FILEWRITE_Append | FILEWRITE_AllowRead));
+		if (LocalFileAr.IsValid())
+		{
+			LocalFileAr->Seek(LocalFileAr->TotalSize());
+
+			// flush checkpoint
+			if (CheckpointRecordingAr.Get() && (CheckpointRecordingAr->Num() > 0))
+			{
+				ELocalFileChunkType ChunkType = ELocalFileChunkType::Checkpoint;
+				*LocalFileAr << ChunkType;
+
+				int64 SavedPos = LocalFileAr->Tell();
+			
+				int32 PlaceholderSize = 0;
+				*LocalFileAr << PlaceholderSize;
+
+				int64 MetadataPos = LocalFileAr->Tell();
+
+				//@todo - use actual checkpoint index?
+				FString Id = FString::Printf(TEXT("checkpoint%ld"), ReplayInfo.DataChunks.Num());
+				*LocalFileAr << Id;
+
+				FString Group = TEXT("checkpoint");
+				*LocalFileAr << Group;
+
+				FString Metadata = FString::Printf(TEXT("%ld"), ReplayInfo.DataChunks.Num());
+				*LocalFileAr << Metadata;
+
+				uint32 Time1 = TimeInMS;
+				*LocalFileAr << Time1;
+
+				uint32 Time2 = TimeInMS;
+				*LocalFileAr << Time2;
+
+				if (SupportsCompression())
+				{
+					TArray<uint8> Compressed;
+					if (CompressBuffer(*CheckpointRecordingAr, Compressed))
+					{
+						int32 CheckpointSize = Compressed.Num();
+						*LocalFileAr << CheckpointSize;
+
+						LocalFileAr->Serialize(Compressed.GetData(), Compressed.Num());
+					}
+					else
+					{
+						UE_LOG(LogLocalFileReplay, Warning, TEXT("FLocalFileNetworkReplayStreamer::FlushStream - CompressBuffer failed"));
+					}
+				}
+				else
+				{
+					int32 CheckpointSize = CheckpointRecordingAr->Num();
+					*LocalFileAr << CheckpointSize;
+
+					LocalFileAr->Serialize(CheckpointRecordingAr->GetData(), CheckpointRecordingAr->Num());
+				}
+
+				int32 ChunkSize = LocalFileAr->Tell() - MetadataPos;
+
+				LocalFileAr->Seek(SavedPos);
+				*LocalFileAr << ChunkSize;
+
+				CheckpointRecordingAr.Reset();
+			}
+		}
+	}
+
+	WriteReplayInfo(CurrentStreamName, ReplayInfo);
+}
+
+void FLocalFileNetworkReplayStreamer::GotoCheckpointIndex(const int32 CheckpointIndex, const FOnCheckpointReadyDelegate& Delegate)
+{
+	GotoCheckpointIndexInternal(CheckpointIndex, Delegate, -1);
+}
+
+void FLocalFileNetworkReplayStreamer::GotoCheckpointIndexInternal(int32 CheckpointIndex, const FOnCheckpointReadyDelegate& Delegate, int32 TimeInMS)
+{
+	if (CheckpointIndex == -1)
+	{
+		// Create a dummy checkpoint archive to indicate this is the first checkpoint
+		CheckpointPlaybackAr.Reset(new FArrayReader);
+		FilePlaybackAr->Seek(0);
+		Delegate.ExecuteIfBound(true, TimeInMS);
+		return;
+	}
+
+	if (!ReplayInfo.Checkpoints.IsValidIndex(CheckpointIndex))
+	{
+		UE_LOG(LogLocalFileReplay, Log, TEXT("FLocalFileNetworkReplayStreamer::GotoCheckpointIndex. Invalid Index: %i"), CheckpointIndex);
+		Delegate.ExecuteIfBound(false, TimeInMS);
+		return;
+	}
+
+	const FString FullFilename = LocalFileReplay::GetDemoFullFilename(CurrentStreamName);
+	
+	TUniquePtr< FArchive > LocalFileAr (IFileManager::Get().CreateFileReader(*FullFilename, FILEREAD_AllowWrite));
+	LocalFileAr->Seek(ReplayInfo.Checkpoints[CheckpointIndex].Offset);
+
+	CheckpointPlaybackAr.Reset(new FArrayReader());
+
+	CheckpointPlaybackAr->AddUninitialized(ReplayInfo.Checkpoints[CheckpointIndex].SizeInBytes);
+	LocalFileAr->Serialize(CheckpointPlaybackAr->GetData(), CheckpointPlaybackAr->Num());
+
+	if (SupportsCompression())
+	{
+		TArray<uint8> Uncompressed;
+		if (DecompressBuffer(*CheckpointPlaybackAr, Uncompressed))
+		{
+			CheckpointPlaybackAr->Empty();
+			CheckpointPlaybackAr->Append(Uncompressed);
+		}
+	}
+
+	check(FilePlaybackAr.Get() != nullptr);
+	FilePlaybackAr->SetChunk(FCString::Atoi64(*ReplayInfo.Checkpoints[CheckpointIndex].Metadata));
+
+	Delegate.ExecuteIfBound(true, TimeInMS);
+}
+
+void FLocalFileNetworkReplayStreamer::GotoTimeInMS(const uint32 TimeInMS, const FOnCheckpointReadyDelegate& Delegate)
+{
+	int32 CheckpointIndex = -1;
+
+	if (ReplayInfo.Checkpoints.Num() > 0 && TimeInMS >= ReplayInfo.Checkpoints[ ReplayInfo.Checkpoints.Num() - 1 ].Time1)
+	{
+		// If we're after the very last checkpoint, that's the one we want
+		CheckpointIndex = ReplayInfo.Checkpoints.Num() - 1;
+	}
+	else
+	{
+		// Checkpoints should be sorted by time, return the checkpoint that exists right before the current time
+		// For fine scrubbing, we'll fast forward the rest of the way
+		// NOTE - If we're right before the very first checkpoint, we'll return -1, which is what we want when we want to start from the very beginning
+		for (int i = 0; i < ReplayInfo.Checkpoints.Num(); i++)
+		{
+			if (TimeInMS < ReplayInfo.Checkpoints[i].Time1)
+			{
+				CheckpointIndex = i - 1;
+				break;
+			}
+		}
+	}
+
+	int32 ExtraSkipTimeInMS = TimeInMS;
+
+	if (CheckpointIndex >= 0)
+	{
+		// Subtract off checkpoint time so we pass in the leftover to the engine to fast forward through for the fine scrubbing part
+		ExtraSkipTimeInMS = TimeInMS - ReplayInfo.Checkpoints[ CheckpointIndex ].Time1;
+	}
+
+	GotoCheckpointIndexInternal(CheckpointIndex, Delegate, ExtraSkipTimeInMS);
+}
+
+void FLocalFileNetworkReplayStreamer::Tick(float DeltaSeconds)
+{
+	if (StreamerState == EStreamerState::Playback)
+	{
+		// update replay info with latest information
+		if (ReplayInfo.bIsLive)
+		{
+			int64 LastDataSize = ReplayInfo.TotalDataSizeInBytes;
+				
+			FLocalFileReplayInfo LatestInfo = ReadReplayInfo(CurrentStreamName);
+			if (LatestInfo.bIsValid)
+			{
+				ReplayInfo = LatestInfo;
+
+				// reopen playback file if size has changed
+				if (FilePlaybackAr.IsValid() && (LastDataSize != ReplayInfo.TotalDataSizeInBytes))
+				{
+					int64 OldPos = FilePlaybackAr->Tell();
+					FilePlaybackAr->LocalFileAr.Reset(IFileManager::Get().CreateFileReader(*LocalFileReplay::GetDemoFullFilename(CurrentStreamName), FILEREAD_AllowWrite));
+					FilePlaybackAr->Seek(OldPos);
+				}
+			}
+		}
+	}
+	else if (StreamerState == EStreamerState::Recording)
+	{
+		// when a replay chunk gets flushed it will update the header
+	}
+}
+
+TStatId FLocalFileNetworkReplayStreamer::GetStatId() const
+{
+	RETURN_QUICK_DECLARE_CYCLE_STAT(FLocalFileNetworkReplayStreamer, STATGROUP_Tickables);
+}
+
+IMPLEMENT_MODULE(FLocalFileNetworkReplayStreamingFactory, LocalFileNetworkReplayStreaming)
+
+TSharedPtr< INetworkReplayStreamer > FLocalFileNetworkReplayStreamingFactory::CreateReplayStreamer() 
+{
+	return TSharedPtr< INetworkReplayStreamer >(new FLocalFileNetworkReplayStreamer);
+}
diff --git a/Engine/Source/Runtime/NetworkReplayStreaming/LocalFileNetworkReplayStreaming/Public/LocalFileNetworkReplayStreaming.h b/Engine/Source/Runtime/NetworkReplayStreaming/LocalFileNetworkReplayStreaming/Public/LocalFileNetworkReplayStreaming.h
new file mode 100644
--- /dev/null
+++ b/Engine/Source/Runtime/NetworkReplayStreaming/LocalFileNetworkReplayStreaming/Public/LocalFileNetworkReplayStreaming.h
@@ -0,0 +1,246 @@
+// Copyright 1998-2017 Epic Games, Inc. All Rights Reserved.
+
+#pragma once
+
+#include "CoreMinimal.h"
+#include "NetworkReplayStreaming.h"
+#include "Stats/Stats.h"
+#include "Tickable.h"
+#include "Serialization/ArrayReader.h"
+#include "Serialization/ArrayWriter.h"
+#include "Serialization/MemoryReader.h"
+
+class FNetworkReplayVersion;
+class FLocalFileNetworkReplayStreamer;
+
+enum class ELocalFileChunkType : uint32
+{
+	Header,
+	ReplayData,
+	Checkpoint,
+	Event,
+	Unknown = 0xFFFFFFFF
+};
+
+/** Struct to hold chunk metadata */
+struct FLocalFileChunkInfo
+{
+	FLocalFileChunkInfo() : 
+		ChunkType(ELocalFileChunkType::Unknown),
+		SizeInBytes(0), 
+		Offset(0) 
+	{}
+
+	ELocalFileChunkType ChunkType;
+	int32 SizeInBytes;
+	int64 Offset;
+};
+
+/** Struct to hold replay data chunk metadata */
+struct FLocalFileReplayDataInfo
+{
+	FLocalFileReplayDataInfo() :
+		ChunkIndex(INDEX_NONE),
+		ReplayOffset(0)
+	{}
+
+	int32 ChunkIndex;
+	int64 ReplayOffset;
+};
+
+/** Struct to hold event metadata */
+struct FLocalFileEventInfo
+{
+	FLocalFileEventInfo() : 
+		ChunkIndex(INDEX_NONE),
+		Time1(0),
+		Time2(0),
+		SizeInBytes(0),
+		Offset(0)
+	{}
+	
+	int32 ChunkIndex;
+
+	FString Id;
+	FString Group;
+	FString Metadata;
+	uint32 Time1;
+	uint32 Time2;
+
+	int32 SizeInBytes;
+	int64 Offset;
+};
+
+/** Struct to hold metadata about an entire replay */
+struct FLocalFileReplayInfo
+{
+	FLocalFileReplayInfo() : 
+		LengthInMS(0), 
+		NetworkVersion(0),
+		Changelist(0),
+		TotalDataSizeInBytes(0),
+		bIsLive(false),
+		bIsValid(false),
+		HeaderChunkIndex(INDEX_NONE)
+	{}
+
+	int32 LengthInMS;
+	uint32 NetworkVersion;
+	uint32 Changelist;
+	FString	FriendlyName;
+	int64 TotalDataSizeInBytes;
+	bool bIsLive;
+	bool bIsValid;
+
+	int32 HeaderChunkIndex;
+
+	TArray<FLocalFileChunkInfo> Chunks;
+	TArray<FLocalFileEventInfo> Checkpoints;
+	TArray<FLocalFileEventInfo> Events;
+	TArray<FLocalFileReplayDataInfo> DataChunks;
+};
+
+/** Archive to wrap the file reader and respect chunk boundaries */
+class LOCALFILENETWORKREPLAYSTREAMING_API FLocalFileStreamFArchive : public FArchive
+{
+public:
+	FLocalFileStreamFArchive(FLocalFileNetworkReplayStreamer* InStreamer) :
+		Streamer(InStreamer), 
+		Pos(0), 
+		CurrentChunkIndex(0), 
+		CurrentChunkPos(0) 
+	{
+	}
+
+	/** FArchive overrides */
+	virtual void	Serialize(void* V, int64 Length) override;
+	virtual int64	Tell() override;
+	virtual int64	TotalSize() override;
+	virtual void	Seek(int64 InPos) override;
+	virtual bool	AtEnd() override;
+
+	/** Seek to the beginning of specified chunk */
+	void			SetChunk(int32 ChunkIndex);
+
+	FLocalFileNetworkReplayStreamer* Streamer;
+	TUniquePtr<FArchive>	LocalFileAr;
+	int64					Pos;
+	int32					CurrentChunkIndex;
+	TArray<uint8>			CurrentChunkBuffer;
+	int32					CurrentChunkPos;
+};
+
+/** Archive to wrap header writes that respects FArchive::Flush */
+class LOCALFILENETWORKREPLAYSTREAMING_API FLocalFileHeaderFArchive : public FArrayWriter
+{
+public:
+	FLocalFileHeaderFArchive(FLocalFileNetworkReplayStreamer* InStreamer) : Streamer(InStreamer) {}
+
+	virtual void Flush() override;
+
+	FLocalFileNetworkReplayStreamer* Streamer;
+};
+
+/** Local file streamer that supports playback/recording to a single file on disk */
+class LOCALFILENETWORKREPLAYSTREAMING_API FLocalFileNetworkReplayStreamer : public INetworkReplayStreamer, public FTickableGameObject
+{
+public:
+	FLocalFileNetworkReplayStreamer() :
+		StreamerState(EStreamerState::Idle)
+	{}
+
+	/** INetworkReplayStreamer implementation */
+	virtual void StartStreaming(const FString& CustomName, const FString& FriendlyName, const TArray< FString >& UserNames, bool bRecord, const FNetworkReplayVersion& ReplayVersion, const FOnStreamReadyDelegate& Delegate) override;
+	virtual void StopStreaming() override;
+	virtual FArchive* GetHeaderArchive() override;
+	virtual FArchive* GetStreamingArchive() override;
+	virtual FArchive* GetCheckpointArchive() override;
+	virtual void FlushCheckpoint(const uint32 TimeInMS) override;
+	virtual void GotoCheckpointIndex(const int32 CheckpointIndex, const FOnCheckpointReadyDelegate& Delegate) override;
+	virtual void GotoTimeInMS(const uint32 TimeInMS, const FOnCheckpointReadyDelegate& Delegate) override;
+	virtual void UpdateTotalDemoTime(uint32 TimeInMS) override;
+	virtual uint32 GetTotalDemoTime() const override { return ReplayInfo.LengthInMS; }
+	virtual bool IsDataAvailable() const override;
+	virtual void SetHighPriorityTimeRange(const uint32 StartTimeInMS, const uint32 EndTimeInMS) override {}
+	virtual bool IsDataAvailableForTimeRange(const uint32 StartTimeInMS, const uint32 EndTimeInMS) override { return true; }
+	virtual bool IsLoadingCheckpoint() const override { return false; }
+	virtual bool IsLive() const override;
+	virtual void DeleteFinishedStream(const FString& StreamName, const FOnDeleteFinishedStreamComplete& Delegate) const override;
+	virtual void EnumerateStreams(const FNetworkReplayVersion& ReplayVersion, const FString& UserString, const FString& MetaString, const FOnEnumerateStreamsComplete& Delegate) override;
+	virtual void EnumerateStreams(const FNetworkReplayVersion& InReplayVersion, const FString& UserString, const FString& MetaString, const TArray< FString >& ExtraParms, const FOnEnumerateStreamsComplete& Delegate) override;
+	virtual void EnumerateRecentStreams(const FNetworkReplayVersion& ReplayVersion, const FString& RecentViewer, const FOnEnumerateStreamsComplete& Delegate) override {}
+	virtual ENetworkReplayError::Type GetLastError() const override { return ENetworkReplayError::None; }
+	virtual void AddUserToReplay(const FString& UserString) override;
+	virtual void AddEvent(const uint32 TimeInMS, const FString& Group, const FString& Meta, const TArray<uint8>& Data) override;
+	virtual void AddOrUpdateEvent(const FString& Name, const uint32 TimeInMS, const FString& Group, const FString& Meta, const TArray<uint8>& Data) override;
+	virtual void EnumerateEvents(const FString& Group, const FEnumerateEventsCompleteDelegate& EnumerationCompleteDelegate) override;
+	virtual void EnumerateEvents(const FString& ReplayName, const FString& Group, const FEnumerateEventsCompleteDelegate& EnumerationCompleteDelegate) override;
+	virtual void RequestEventData(const FString& EventID, const FOnRequestEventDataComplete& RequestEventDataComplete) override;
+	virtual void SearchEvents(const FString& EventGroup, const FOnEnumerateStreamsComplete& Delegate) override;
+	virtual void KeepReplay(const FString& ReplayName, const bool bKeep) override {}
+	virtual FString	GetReplayID() const override { return TEXT(""); }
+	virtual void SetTimeBufferHintSeconds(const float InTimeBufferHintSeconds) override {}
+	virtual void RefreshHeader() override {};
+	virtual void DownloadHeader(const FOnDownloadHeaderComplete& Delegate = FOnDownloadHeaderComplete()) override {}
+
+	virtual bool SupportsCompression() const { return false; }
+	virtual int32 GetDecompressedSize(FArchive& InCompressed) const { return 0; }
+	virtual bool DecompressBuffer(const TArray<uint8>& InCompressed, TArray< uint8 >& OutBuffer) const { return false; }
+	virtual bool CompressBuffer(const TArray< uint8 >& InBuffer, TArray< uint8 >& OutCompressed) const { return false; }
+
+	/** FTickableObjectBase implementation */
+	virtual void Tick(float DeltaSeconds) override;
+	virtual bool IsTickable() const override { return true; }
+	virtual TStatId GetStatId() const override;
+
+	/** FTickableGameObject implementation */
+	virtual bool IsTickableWhenPaused() const override { return true; }
+
+	void FlushHeader();
+
+	/** Currently playing or recording replay metadata */
+	FLocalFileReplayInfo ReplayInfo;
+
+protected:
+	FLocalFileReplayInfo ReadReplayInfo(const FString& StreamName) const;
+	void WriteReplayInfo(const FString& StreamName, FLocalFileReplayInfo& ReplayInfo);
+
+	bool IsNamedStreamLive(const FString& StreamName) const;
+
+	/** Handles the details of loading a checkpoint */
+	void GotoCheckpointIndexInternal(int32 CheckpointIndex, const FOnCheckpointReadyDelegate& Delegate, int32 TimeInMS);
+
+	void FlushStream();
+
+	/** Handle to the archive that will read/write the demo header */
+	TUniquePtr<FArchive> HeaderPlaybackAr;
+	TUniquePtr<FLocalFileHeaderFArchive> HeaderRecordingAr;
+
+	/** Handle to the archive that will read/write network packets */
+	TUniquePtr<FLocalFileStreamFArchive> FilePlaybackAr;
+	TUniquePtr<FArrayWriter> FileRecordingAr;
+
+	/* Handle to the archive that will read/write checkpoint files */
+	TUniquePtr<FArrayReader> CheckpointPlaybackAr;
+	TUniquePtr<FArrayWriter> CheckpointRecordingAr;
+
+	/** EStreamerState - Overall state of the streamer */
+	enum class EStreamerState
+	{
+		Idle,					// The streamer is idle. Either we haven't started streaming yet, or we are done
+		Recording,				// We are in the process of recording a replay to disk
+		Playback,				// We are in the process of playing a replay from disk
+	};
+
+	/** Overall state of the streamer */
+	EStreamerState StreamerState;
+
+	/** Remember the name of the current stream, if any. */
+	FString CurrentStreamName;
+};
+
+class LOCALFILENETWORKREPLAYSTREAMING_API FLocalFileNetworkReplayStreamingFactory : public INetworkReplayStreamingFactory
+{
+public:
+	virtual TSharedPtr< INetworkReplayStreamer > CreateReplayStreamer();
+};
diff --git a/Engine/Source/Runtime/Sockets/Private/NetworkMessage.cpp b/Engine/Source/Runtime/Sockets/Private/NetworkMessage.cpp
--- a/Engine/Source/Runtime/Sockets/Private/NetworkMessage.cpp
+++ b/Engine/Source/Runtime/Sockets/Private/NetworkMessage.cpp
@@ -27,13 +27,18 @@ bool FSimpleAbstractSocket_FSocket::Receive(uint8 *Results, int32 Size) const
 
 bool FSimpleAbstractSocket_FSocket::Send(const uint8 *Buffer, int32 Size) const
 {
-	int32 AmountSent = 0;
-	Socket->Send(Buffer, Size, AmountSent);
-	// make sure it sent it all
-	if (AmountSent != Size)
+	while (Size > 0)
 	{
-		return false;
+		int32 BytesSent = 0;
+		if (!Socket->Send(Buffer, Size, BytesSent))
+		{
+			return false;
+		}
+
+		Size -= BytesSent;
+		Buffer += BytesSent;
 	}
+
 	return true;
 }
 
