Fix off by one issue with automatic demo filenames in null streamer

#rnx

diff --git a/Engine/Source/Runtime/NetworkReplayStreaming/NullNetworkReplayStreaming/Private/NullNetworkReplayStreaming.cpp b/Engine/Source/Runtime/NetworkReplayStreaming/NullNetworkReplayStreaming/Private/NullNetworkReplayStreaming.cpp
--- a/Engine/Source/Runtime/NetworkReplayStreaming/NullNetworkReplayStreaming/Private/NullNetworkReplayStreaming.cpp
+++ b/Engine/Source/Runtime/NetworkReplayStreaming/NullNetworkReplayStreaming/Private/NullNetworkReplayStreaming.cpp
@@ -152,7 +152,7 @@ FString GetAutomaticDemoName()
 	int32 i = 1;
 	while(bUnlimitedDemos || i <= MaxDemos)
 	{
-		const FString DemoName = FString::Printf(TEXT("demo%i"), i + 1);
+		const FString DemoName = FString::Printf(TEXT("demo%i"), i);
 		const FString FullDemoName = GetDemoFilename(DemoName);
 		
 		FDateTime DateTime = IFileManager::Get().GetTimeStamp(*FullDemoName);
