PR #4875: Remove extra Printf parameter

diff --git a/Engine/Source/Runtime/NetworkReplayStreaming/HttpNetworkReplayStreaming/Private/HttpNetworkReplayStreaming.cpp b/Engine/Source/Runtime/NetworkReplayStreaming/HttpNetworkReplayStreaming/Private/HttpNetworkReplayStreaming.cpp
--- a/Engine/Source/Runtime/NetworkReplayStreaming/HttpNetworkReplayStreaming/Private/HttpNetworkReplayStreaming.cpp
+++ b/Engine/Source/Runtime/NetworkReplayStreaming/HttpNetworkReplayStreaming/Private/HttpNetworkReplayStreaming.cpp
@@ -1093,7 +1093,7 @@ void FHttpNetworkReplayStreamer::DownloadHeader(const FDownloadHeaderCallback& D
 	// Download header first
 	TSharedRef<IHttpRequest> HttpRequest = FHttpModule::Get().CreateRequest();
 
-	HttpRequest->SetURL( FString::Printf( TEXT( "%sreplay/%s/file/replay.header" ), *ServerURL, *SessionName, *ViewerName ) );
+	HttpRequest->SetURL( FString::Printf( TEXT( "%sreplay/%s/file/replay.header" ), *ServerURL, *SessionName ) );
 	HttpRequest->SetVerb( TEXT( "GET" ) );
 
 	HttpRequest->OnProcessRequestComplete().BindRaw( this, &FHttpNetworkReplayStreamer::HttpDownloadHeaderFinished, Delegate );
