[LLM] created a new Replay tag to further tidy EngineMisc

diff --git a/Engine/Source/Runtime/Core/Private/HAL/LowLevelMemTracker.cpp b/Engine/Source/Runtime/Core/Private/HAL/LowLevelMemTracker.cpp
--- a/Engine/Source/Runtime/Core/Private/HAL/LowLevelMemTracker.cpp
+++ b/Engine/Source/Runtime/Core/Private/HAL/LowLevelMemTracker.cpp
@@ -107,6 +107,7 @@ DECLARE_LLM_MEMORY_STAT(TEXT("AssetRegistry"), STAT_AssetRegistryLLM, STATGROUP_
 DECLARE_LLM_MEMORY_STAT(TEXT("ConfigSystem"), STAT_ConfigSystemLLM, STATGROUP_LLMFULL);
 DECLARE_LLM_MEMORY_STAT(TEXT("InitUObject"), STAT_InitUObjectLLM, STATGROUP_LLMFULL);
 DECLARE_LLM_MEMORY_STAT(TEXT("VideoRecording"), STAT_VideoRecordingLLM, STATGROUP_LLMFULL);
+DECLARE_LLM_MEMORY_STAT(TEXT("Replays"), STAT_ReplaysLLM, STATGROUP_LLMFULL);
 DECLARE_LLM_MEMORY_STAT(TEXT("CsvProfiler"), STAT_CsvProfilerLLM, STATGROUP_LLMFULL);
 DECLARE_LLM_MEMORY_STAT(TEXT("MaterialInstance"), STAT_MaterialInstanceLLM, STATGROUP_LLMFULL);
 DECLARE_LLM_MEMORY_STAT(TEXT("VertexData"), STAT_VertexDataLLM, STATGROUP_LLMFULL);
diff --git a/Engine/Source/Runtime/Core/Public/HAL/LowLevelMemTracker.h b/Engine/Source/Runtime/Core/Public/HAL/LowLevelMemTracker.h
--- a/Engine/Source/Runtime/Core/Public/HAL/LowLevelMemTracker.h
+++ b/Engine/Source/Runtime/Core/Public/HAL/LowLevelMemTracker.h
@@ -233,6 +233,7 @@ enum class ELLMTagSet : uint8
 	macro(ConfigSystem,							"ConfigSystem",					GET_STATFNAME(STAT_ConfigSystemLLM),						GET_STATFNAME(STAT_EngineSummaryLLM),			-1)\
 	macro(InitUObject,							"InitUObject",					GET_STATFNAME(STAT_InitUObjectLLM),							GET_STATFNAME(STAT_EngineSummaryLLM),			-1)\
 	macro(VideoRecording,						"VideoRecording",				GET_STATFNAME(STAT_VideoRecordingLLM),						GET_STATFNAME(STAT_EngineSummaryLLM),			-1)\
+	macro(Replays,								"Replays",						GET_STATFNAME(STAT_ReplaysLLM),								GET_STATFNAME(STAT_EngineSummaryLLM),			-1)\
 	macro(MaterialInstance,						"MaterialInstance",				GET_STATFNAME(STAT_MaterialInstanceLLM),					GET_STATFNAME(STAT_EngineSummaryLLM),			-1)\
 	macro(VertexData,							"VertexData",					GET_STATFNAME(STAT_VertexDataLLM),							GET_STATFNAME(STAT_EngineSummaryLLM),			-1)\
 	macro(SkeletalMesh,							"SkeletalMesh",					GET_STATFNAME(STAT_SkeletalMeshLLM),						GET_STATFNAME(STAT_EngineSummaryLLM),			ELLMTag::Meshes)\
diff --git a/Engine/Source/Runtime/NetworkReplayStreaming/LocalFileNetworkReplayStreaming/Private/LocalFileNetworkReplayStreaming.cpp b/Engine/Source/Runtime/NetworkReplayStreaming/LocalFileNetworkReplayStreaming/Private/LocalFileNetworkReplayStreaming.cpp
--- a/Engine/Source/Runtime/NetworkReplayStreaming/LocalFileNetworkReplayStreaming/Private/LocalFileNetworkReplayStreaming.cpp
+++ b/Engine/Source/Runtime/NetworkReplayStreaming/LocalFileNetworkReplayStreaming/Private/LocalFileNetworkReplayStreaming.cpp
@@ -3,6 +3,7 @@
 #include "LocalFileNetworkReplayStreaming.h"
 #include "Misc/NetworkVersion.h"
 #include "HAL/FileManager.h"
+#include "HAL/LowLevelMemTracker.h"
 #include "Misc/Paths.h"
 #include "Misc/Guid.h"
 #include "Async/Async.h"
@@ -2146,6 +2147,8 @@ bool FLocalFileNetworkReplayStreamer::IsFileRequestPendingOrInProgress(const EQu
 
 bool FLocalFileNetworkReplayStreamer::ProcessNextFileRequest()
 {
+	LLM_SCOPE(ELLMTag::Replays);
+
 	if (IsFileRequestInProgress())
 	{
 		return false;
@@ -2243,6 +2246,8 @@ const TArray<uint8>& FLocalFileNetworkReplayStreamer::GetCachedFileContents(cons
 
 TSharedPtr<FArchive> FLocalFileNetworkReplayStreamer::CreateLocalFileReader(const FString& InFilename) const
 {
+	LLM_SCOPE(ELLMTag::Replays);
+
 	if (bCacheFileReadsInMemory)
 	{
 		const TArray<uint8>& Data = GetCachedFileContents(InFilename);
@@ -2656,6 +2661,7 @@ void FLocalFileNetworkReplayStreamer::ConditionallyLoadNextChunk()
 		[this, RequestedStreamChunkIndex](TLocalFileRequestCommonData<FStreamingResultBase>& RequestData)
 		{
 			SCOPE_CYCLE_COUNTER(STAT_LocalReplay_ReadStream);
+			LLM_SCOPE(ELLMTag::Replays);
 
 			if (ReadReplayInfo(CurrentStreamName, RequestData.ReplayInfo))
 			{
@@ -2704,6 +2710,8 @@ void FLocalFileNetworkReplayStreamer::ConditionallyLoadNextChunk()
 		},
 		[this, RequestedStreamChunkIndex](TLocalFileRequestCommonData<FStreamingResultBase>& RequestData)
 		{
+			LLM_SCOPE(ELLMTag::Replays);
+
 			// Make sure our stream chunk index didn't change under our feet
 			if (RequestedStreamChunkIndex != StreamChunkIndex)
 			{
@@ -2817,6 +2825,7 @@ void FLocalFileNetworkReplayStreamer::AddRequestToCache(int32 ChunkIndex, const
 	}
 
 	// Add to cache (or freshen existing entry)
+	LLM_SCOPE(ELLMTag::Replays);
 	RequestCache.Add(ChunkIndex, MakeShareable(new FCachedFileRequest(RequestData, FPlatformTime::Seconds())));
 
 	// Anytime we add something to cache, make sure it's within budget
