Fix checkpoint loading crash in http replay streamer

diff --git a/Engine/Source/Runtime/NetworkReplayStreaming/HttpNetworkReplayStreaming/Private/HttpNetworkReplayStreaming.cpp b/Engine/Source/Runtime/NetworkReplayStreaming/HttpNetworkReplayStreaming/Private/HttpNetworkReplayStreaming.cpp
--- a/Engine/Source/Runtime/NetworkReplayStreaming/HttpNetworkReplayStreaming/Private/HttpNetworkReplayStreaming.cpp
+++ b/Engine/Source/Runtime/NetworkReplayStreaming/HttpNetworkReplayStreaming/Private/HttpNetworkReplayStreaming.cpp
@@ -838,17 +838,20 @@ void FHttpNetworkReplayStreamer::GotoCheckpointIndex( const int32 CheckpointInde
 	TSharedRef<IHttpRequest> HttpRequest = FHttpModule::Get().CreateRequest();
 
 	// Download the next stream chunk
-	HttpRequest->SetURL( FString::Printf( TEXT( "%sevent/%s" ), *ServerURL,  *CheckpointList.ReplayEvents[DeltaDownloadCheckpointIndex].ID ) );
 	HttpRequest->SetVerb( TEXT( "GET" ) );
 
 	if (CheckpointType == EReplayCheckpointType::Delta)
 	{
 		DeltaDownloadCheckpointIndex = 0;
 
+		HttpRequest->SetURL( FString::Printf( TEXT( "%sevent/%s" ), *ServerURL,  *CheckpointList.ReplayEvents[DeltaDownloadCheckpointIndex].ID ) );
+
 		HttpRequest->OnProcessRequestComplete().BindRaw( this, &FHttpNetworkReplayStreamer::HttpDownloadCheckpointDeltaFinished );
 	}
 	else
 	{
+		HttpRequest->SetURL( FString::Printf( TEXT( "%sevent/%s" ), *ServerURL,  *CheckpointList.ReplayEvents[CheckpointIndex].ID ) );
+
 		HttpRequest->OnProcessRequestComplete().BindRaw( this, &FHttpNetworkReplayStreamer::HttpDownloadCheckpointFinished );
 	}
 
