Add flush for loaded replay streamers/factories

diff --git a/Engine/Source/Runtime/NetworkReplayStreaming/LocalFileNetworkReplayStreaming/Private/LocalFileNetworkReplayStreaming.cpp b/Engine/Source/Runtime/NetworkReplayStreaming/LocalFileNetworkReplayStreaming/Private/LocalFileNetworkReplayStreaming.cpp
--- a/Engine/Source/Runtime/NetworkReplayStreaming/LocalFileNetworkReplayStreaming/Private/LocalFileNetworkReplayStreaming.cpp
+++ b/Engine/Source/Runtime/NetworkReplayStreaming/LocalFileNetworkReplayStreaming/Private/LocalFileNetworkReplayStreaming.cpp
@@ -1,3 +1,4 @@
+
 // Copyright 1998-2019 Epic Games, Inc. All Rights Reserved.
 
 #include "LocalFileNetworkReplayStreaming.h"
@@ -3060,7 +3061,23 @@ void FLocalFileNetworkReplayStreamingFactory::Tick( float DeltaTime )
 	}
 }
 
-void FLocalFileNetworkReplayStreamingFactory::ShutdownModule()
+bool FLocalFileNetworkReplayStreamingFactory::HasAnyPendingRequests() const
+{
+	bool bPendingRequests = false;
+
+	for (const TSharedPtr<FLocalFileNetworkReplayStreamer>& Streamer : LocalFileStreamers)
+	{
+		if (Streamer.IsValid() && Streamer->HasPendingFileRequests())
+		{
+			bPendingRequests = true;
+			break;
+		}
+	}
+
+	return bPendingRequests;
+}
+
+void FLocalFileNetworkReplayStreamingFactory::Flush()
 {
 	bool bFlushStreamersOnShutdown = true;
 	GConfig->GetBool(TEXT("LocalFileNetworkReplayStreamingFactory"), TEXT("bFlushStreamersOnShutdown"), bFlushStreamersOnShutdown, GEngineIni);
@@ -3072,7 +3089,7 @@ void FLocalFileNetworkReplayStreamingFactory::ShutdownModule()
 
 		double BeginWaitTime = FPlatformTime::Seconds();
 		double LastTime = BeginWaitTime;
-		while (LocalFileStreamers.Num() > 0)
+		while (HasAnyPendingRequests())
 		{
 			const double AppTime = FPlatformTime::Seconds();
 			const double TotalWait = AppTime - BeginWaitTime;
@@ -3086,7 +3103,7 @@ void FLocalFileNetworkReplayStreamingFactory::ShutdownModule()
 			Tick(AppTime - LastTime);
 			LastTime = AppTime;
 
-			if (LocalFileStreamers.Num() > 0)
+			if (HasAnyPendingRequests())
 			{
 				FTaskGraphInterface::Get().ProcessThreadUntilIdle(ENamedThreads::GameThread);
 
@@ -3100,6 +3117,11 @@ void FLocalFileNetworkReplayStreamingFactory::ShutdownModule()
 	}
 }
 
+void FLocalFileNetworkReplayStreamingFactory::ShutdownModule()
+{
+	Flush();
+}
+
 TStatId FLocalFileNetworkReplayStreamingFactory::GetStatId() const
 {
 	RETURN_QUICK_DECLARE_CYCLE_STAT(FLocalFileNetworkReplayStreamingFactory, STATGROUP_Tickables);
diff --git a/Engine/Source/Runtime/NetworkReplayStreaming/LocalFileNetworkReplayStreaming/Public/LocalFileNetworkReplayStreaming.h b/Engine/Source/Runtime/NetworkReplayStreaming/LocalFileNetworkReplayStreaming/Public/LocalFileNetworkReplayStreaming.h
--- a/Engine/Source/Runtime/NetworkReplayStreaming/LocalFileNetworkReplayStreaming/Public/LocalFileNetworkReplayStreaming.h
+++ b/Engine/Source/Runtime/NetworkReplayStreaming/LocalFileNetworkReplayStreaming/Public/LocalFileNetworkReplayStreaming.h
@@ -644,7 +644,8 @@ class LOCALFILENETWORKREPLAYSTREAMING_API FLocalFileNetworkReplayStreamingFactor
 public:
 	virtual void ShutdownModule() override;
 
-	virtual TSharedPtr<INetworkReplayStreamer> CreateReplayStreamer();
+	virtual TSharedPtr<INetworkReplayStreamer> CreateReplayStreamer() override;
+	virtual void Flush() override;
 
 	/** FTickableGameObject */
 	virtual void Tick(float DeltaTime) override;
@@ -653,5 +654,7 @@ public:
 	bool IsTickableWhenPaused() const override { return true; }
 
 protected:
+	bool HasAnyPendingRequests() const;
+
 	TArray<TSharedPtr<FLocalFileNetworkReplayStreamer>> LocalFileStreamers;
 };
diff --git a/Engine/Source/Runtime/NetworkReplayStreaming/NetworkReplayStreaming/Private/NetworkReplayStreaming.cpp b/Engine/Source/Runtime/NetworkReplayStreaming/NetworkReplayStreaming/Private/NetworkReplayStreaming.cpp
--- a/Engine/Source/Runtime/NetworkReplayStreaming/NetworkReplayStreaming/Private/NetworkReplayStreaming.cpp
+++ b/Engine/Source/Runtime/NetworkReplayStreaming/NetworkReplayStreaming/Private/NetworkReplayStreaming.cpp
@@ -9,11 +9,11 @@ IMPLEMENT_MODULE( FNetworkReplayStreaming, NetworkReplayStreaming );
 
 INetworkReplayStreamingFactory& FNetworkReplayStreaming::GetFactory(const TCHAR* FactoryNameOverride)
 {
-	FString FactoryName = TEXT( "NullNetworkReplayStreaming" );
+	FString FactoryName = TEXT("NullNetworkReplayStreaming");
 
 	if (FactoryNameOverride == nullptr)
 	{
-		GConfig->GetString( TEXT( "NetworkReplayStreaming" ), TEXT( "DefaultFactoryName" ), FactoryName, GEngineIni );
+		GConfig->GetString(TEXT("NetworkReplayStreaming"), TEXT("DefaultFactoryName"), FactoryName, GEngineIni);
 	}
 	else
 	{
@@ -27,17 +27,19 @@ INetworkReplayStreamingFactory& FNetworkReplayStreaming::GetFactory(const TCHAR*
 	}
 
 	// See if we need to forcefully fallback to the null streamer
-	if ( !FModuleManager::Get().IsModuleLoaded( *FactoryName ) )
+	if (!FModuleManager::Get().IsModuleLoaded(*FactoryName))
 	{
-		FModuleManager::Get().LoadModule( *FactoryName );
+		FModuleManager::Get().LoadModule(*FactoryName);
 	
-		if ( !FModuleManager::Get().IsModuleLoaded( *FactoryName ) )
+		if (!FModuleManager::Get().IsModuleLoaded(*FactoryName))
 		{
-			FactoryName = TEXT( "NullNetworkReplayStreaming" );
+			FactoryName = TEXT("NullNetworkReplayStreaming");
 		}
 	}
 
-	return FModuleManager::Get().LoadModuleChecked< INetworkReplayStreamingFactory >( *FactoryName );
+	LoadedFactories.Add(*FactoryName);
+
+	return FModuleManager::Get().LoadModuleChecked<INetworkReplayStreamingFactory>(*FactoryName);
 }
 
 int32 FNetworkReplayStreaming::GetMaxNumberOfAutomaticReplays()
@@ -85,4 +87,16 @@ const FString FNetworkReplayStreaming::GetAutomaticReplayPrefixExtern() const
 const int32 FNetworkReplayStreaming::GetMaxNumberOfAutomaticReplaysExtern() const
 {
 	return GetMaxNumberOfAutomaticReplays();
+}
+
+void FNetworkReplayStreaming::Flush()
+{
+	for (const FName& FactoryName : LoadedFactories)
+	{
+		if (FModuleManager::Get().IsModuleLoaded(FactoryName))
+		{
+			INetworkReplayStreamingFactory& ReplayFactory = FModuleManager::Get().LoadModuleChecked<INetworkReplayStreamingFactory>(FactoryName);
+			ReplayFactory.Flush();
+		}
+	}
 }
\ No newline at end of file
diff --git a/Engine/Source/Runtime/NetworkReplayStreaming/NetworkReplayStreaming/Public/NetworkReplayStreaming.h b/Engine/Source/Runtime/NetworkReplayStreaming/NetworkReplayStreaming/Public/NetworkReplayStreaming.h
--- a/Engine/Source/Runtime/NetworkReplayStreaming/NetworkReplayStreaming/Public/NetworkReplayStreaming.h
+++ b/Engine/Source/Runtime/NetworkReplayStreaming/NetworkReplayStreaming/Public/NetworkReplayStreaming.h
@@ -610,6 +610,7 @@ class INetworkReplayStreamingFactory : public IModuleInterface
 {
 public:
 	virtual TSharedPtr< INetworkReplayStreamer > CreateReplayStreamer() = 0;
+	virtual void Flush() {}
 };
 
 /** Replay streaming factory manager */
@@ -621,6 +622,8 @@ public:
 		return FModuleManager::LoadModuleChecked< FNetworkReplayStreaming >( "NetworkReplayStreaming" );
 	}
 
+	NETWORKREPLAYSTREAMING_API void Flush();
+
 	virtual INetworkReplayStreamingFactory& GetFactory(const TCHAR* FactoryNameOverride = nullptr);
 
 	virtual const FString GetAutomaticReplayPrefixExtern() const;
@@ -638,4 +641,7 @@ public:
 	// Gets the configured value for whether or not we should use FDateTime::Now as the automatic replay postfix.
 	// If false, it's up to the streamer to determine a proper postfix.
 	static NETWORKREPLAYSTREAMING_API bool UseDateTimeAsAutomaticReplayPostfix();
+
+private:
+	TSet<FName> LoadedFactories;
 };
