Fixing some thread safety issues with SetLastError

diff --git a/Engine/Source/Runtime/NetworkReplayStreaming/LocalFileNetworkReplayStreaming/Private/LocalFileNetworkReplayStreaming.cpp b/Engine/Source/Runtime/NetworkReplayStreaming/LocalFileNetworkReplayStreaming/Private/LocalFileNetworkReplayStreaming.cpp
--- a/Engine/Source/Runtime/NetworkReplayStreaming/LocalFileNetworkReplayStreaming/Private/LocalFileNetworkReplayStreaming.cpp
+++ b/Engine/Source/Runtime/NetworkReplayStreaming/LocalFileNetworkReplayStreaming/Private/LocalFileNetworkReplayStreaming.cpp
@@ -771,7 +771,7 @@ void FLocalFileNetworkReplayStreamer::CancelStreamingRequests()
 	bStopStreamingCalled = false;
 }
 
-void FLocalFileNetworkReplayStreamer::SetLastError( const ENetworkReplayError::Type InLastError )
+void FLocalFileNetworkReplayStreamer::SetLastError(const ENetworkReplayError::Type InLastError)
 {
 	CancelStreamingRequests();
 	StreamerLastError = InLastError;
@@ -1492,7 +1492,8 @@ void FLocalFileNetworkReplayStreamer::FlushStream(const uint32 TimeInMS)
 					if (!CompressBuffer(StreamData, FinalData))
 					{
 						UE_LOG(LogLocalFileReplay, Warning, TEXT("FLocalFileNetworkReplayStreamer::FlushStream - CompressBuffer failed"));
-						SetLastError( ENetworkReplayError::ServiceUnavailable );
+						ReplayInfo.bIsValid = false;
+						return;
 					}
 				}
 				else
@@ -1543,6 +1544,10 @@ void FLocalFileNetworkReplayStreamer::FlushStream(const uint32 TimeInMS)
 				CurrentReplayInfo = ReplayInfo;
 				CurrentReplayInfo.LengthInMS = TotalLengthInMS;
 			}
+			else
+			{
+				SetLastError(ENetworkReplayError::ServiceUnavailable);
+			}
 		});
 
 	StreamAr.Buffer.Empty();
@@ -1608,7 +1613,8 @@ void FLocalFileNetworkReplayStreamer::FlushCheckpointInternal(const uint32 TimeI
 						if (!CompressBuffer(CheckpointData, FinalData))
 						{
 							UE_LOG(LogLocalFileReplay, Warning, TEXT("FLocalFileNetworkReplayStreamer::FlushStream - CompressBuffer failed"));
-							SetLastError( ENetworkReplayError::ServiceUnavailable );
+							ReplayInfo.bIsValid = false;
+							return;
 						}
 					}
 					else
@@ -1674,6 +1680,10 @@ void FLocalFileNetworkReplayStreamer::FlushCheckpointInternal(const uint32 TimeI
 				CurrentReplayInfo = ReplayInfo;
 				CurrentReplayInfo.LengthInMS = CurrentTotalLengthInMS;
 			}
+			else
+			{
+				SetLastError(ENetworkReplayError::ServiceUnavailable);
+			}
 		});
 
 	CheckpointAr.Buffer.Empty();
@@ -2478,7 +2488,7 @@ void FLocalFileNetworkReplayStreamer::WriteHeader()
 					if (ReadReplayInfo(CurrentStreamName, TestInfo))
 					{
 						UE_LOG(LogLocalFileReplay, Error, TEXT("FLocalFileNetworkReplayStreamer::WriteHeader - Current file already has unexpected header"));
-						SetLastError( ENetworkReplayError::ServiceUnavailable );
+						ReplayInfo.bIsValid = false;
 						return;
 					}
 
@@ -2524,6 +2534,10 @@ void FLocalFileNetworkReplayStreamer::WriteHeader()
 				CurrentReplayInfo = ReplayInfo;
 				CurrentReplayInfo.LengthInMS = TotalLengthInMS;
 			}
+			else
+			{
+				SetLastError(ENetworkReplayError::ServiceUnavailable);
+			}
 		});
 
 	// We're done with the header archive
@@ -2718,7 +2732,7 @@ void FLocalFileNetworkReplayStreamer::ConditionallyLoadNextChunk()
 			{
 				StreamAr.Buffer.Empty();
 				StreamAr.Pos = 0;
-				SetLastError( ENetworkReplayError::ServiceUnavailable );
+				SetLastError(ENetworkReplayError::ServiceUnavailable);
 				return;
 			}
 
