Adding analytics support for replay streamers

diff --git a/Engine/Source/Runtime/Engine/Classes/Engine/DemoNetDriver.h b/Engine/Source/Runtime/Engine/Classes/Engine/DemoNetDriver.h
--- a/Engine/Source/Runtime/Engine/Classes/Engine/DemoNetDriver.h
+++ b/Engine/Source/Runtime/Engine/Classes/Engine/DemoNetDriver.h
@@ -623,6 +623,8 @@ public:
 
 	virtual void InitDestroyedStartupActors() override;
 
+	virtual void SetAnalyticsProvider(TSharedPtr<IAnalyticsProvider> InProvider) override;
+
 protected:
 	virtual UChannel* InternalCreateChannelByName(const FName& ChName) override;
 
diff --git a/Engine/Source/Runtime/Engine/Classes/Engine/NetDriver.h b/Engine/Source/Runtime/Engine/Classes/Engine/NetDriver.h
--- a/Engine/Source/Runtime/Engine/Classes/Engine/NetDriver.h
+++ b/Engine/Source/Runtime/Engine/Classes/Engine/NetDriver.h
@@ -735,7 +735,7 @@ public:
 	void UpdateStandbyCheatStatus(void);
 
 	/** Sets the analytics provider */
-	void ENGINE_API SetAnalyticsProvider(TSharedPtr<IAnalyticsProvider> InProvider);
+	virtual void ENGINE_API SetAnalyticsProvider(TSharedPtr<IAnalyticsProvider> InProvider);
 
 #if DO_ENABLE_NET_TEST
 	FPacketSimulationSettings	PacketSimulationSettings;
diff --git a/Engine/Source/Runtime/Engine/Private/DemoNetDriver.cpp b/Engine/Source/Runtime/Engine/Private/DemoNetDriver.cpp
--- a/Engine/Source/Runtime/Engine/Private/DemoNetDriver.cpp
+++ b/Engine/Source/Runtime/Engine/Private/DemoNetDriver.cpp
@@ -6499,4 +6499,14 @@ void UDemoNetConnection::Serialize(FArchive& Ar)
 			QueuedPacket.CountBytes(Ar);
 		}
 	}
-}
\ No newline at end of file
+}
+
+void UDemoNetDriver::SetAnalyticsProvider(TSharedPtr<IAnalyticsProvider> InProvider)
+{
+	Super::SetAnalyticsProvider(InProvider);
+
+	if (ReplayStreamer.IsValid())
+	{
+		ReplayStreamer->SetAnalyticsProvider(InProvider);
+	}
+}
diff --git a/Engine/Source/Runtime/NetworkReplayStreaming/NetworkReplayStreaming/Public/NetworkReplayStreaming.h b/Engine/Source/Runtime/NetworkReplayStreaming/NetworkReplayStreaming/Public/NetworkReplayStreaming.h
--- a/Engine/Source/Runtime/NetworkReplayStreaming/NetworkReplayStreaming/Public/NetworkReplayStreaming.h
+++ b/Engine/Source/Runtime/NetworkReplayStreaming/NetworkReplayStreaming/Public/NetworkReplayStreaming.h
@@ -10,6 +10,7 @@
 #include "Serialization/JsonSerializerMacros.h"
 
 class FNetworkReplayVersion;
+class IAnalyticsProvider;
 
 class FReplayEventListItem : public FJsonSerializable
 {
@@ -603,6 +604,8 @@ public:
 	 * Note, this will always fail for streamers that don't support replays stored on disk.
 	 */
 	virtual EStreamingOperationResult GetDemoPath(FString& DemoPath) const = 0;
+
+	virtual void SetAnalyticsProvider(TSharedPtr<IAnalyticsProvider>& InProvider) {}
 };
 
 /** Replay streamer factory */
