Adding an exec path for the replay streamers

diff --git a/Engine/Source/Runtime/Engine/Engine.Build.cs b/Engine/Source/Runtime/Engine/Engine.Build.cs
--- a/Engine/Source/Runtime/Engine/Engine.Build.cs
+++ b/Engine/Source/Runtime/Engine/Engine.Build.cs
@@ -13,7 +13,7 @@ public class Engine : ModuleRules
 
 		SharedPCHHeaderFile = "Public/EngineSharedPCH.h";
 
-		PublicIncludePathModuleNames.AddRange(new string[] { "Renderer", "PacketHandler", "NetworkReplayStreaming", "AudioMixer", "AnimationCore" });
+		PublicIncludePathModuleNames.AddRange(new string[] { "Renderer", "PacketHandler", "AudioMixer", "AnimationCore" });
 
 		PrivateIncludePaths.AddRange(
 			new string[] {
@@ -32,7 +32,6 @@ public class Engine : ModuleRules
 				"EyeTracker",
 				"MRMesh",
 				"Advertising",
-				"NetworkReplayStreaming",
 				"MovieSceneCapture",
 				"AutomationWorker",
 				"MovieSceneCapture",
@@ -84,8 +83,9 @@ public class Engine : ModuleRules
 				"AudioPlatformConfiguration",
 				"MeshDescription",
 				"PakFile",
-			}
-		);
+                "NetworkReplayStreaming",
+            }
+        );
 
 		PrivateDependencyModuleNames.AddRange(
 			new string[] {
@@ -302,7 +302,6 @@ public class Engine : ModuleRules
 
 		DynamicallyLoadedModuleNames.AddRange(
 			new string[] {
-				"NetworkReplayStreaming",
 				"NullNetworkReplayStreaming",
 				"HttpNetworkReplayStreaming",
 				"Advertising"
diff --git a/Engine/Source/Runtime/Engine/Private/DemoNetDriver.cpp b/Engine/Source/Runtime/Engine/Private/DemoNetDriver.cpp
--- a/Engine/Source/Runtime/Engine/Private/DemoNetDriver.cpp
+++ b/Engine/Source/Runtime/Engine/Private/DemoNetDriver.cpp
@@ -626,7 +626,10 @@ UDemoNetDriver::UDemoNetDriver(const FObjectInitializer& ObjectInitializer)
 	bIsWaitingForStream = false;
 	MaxArchiveReadPos = 0;
 
-	LevelIntervals.Reserve(512);
+	if (!HasAnyFlags(RF_ClassDefaultObject))
+	{
+		LevelIntervals.Reserve(512);
+	}
 
 	RecordBuildConsiderAndPrioritizeTimeSlice = CVarDemoMaximumRepPrioritizeTime.GetValueOnGameThread();
 }
diff --git a/Engine/Source/Runtime/Launch/Private/LaunchEngineLoop.cpp b/Engine/Source/Runtime/Launch/Private/LaunchEngineLoop.cpp
--- a/Engine/Source/Runtime/Launch/Private/LaunchEngineLoop.cpp
+++ b/Engine/Source/Runtime/Launch/Private/LaunchEngineLoop.cpp
@@ -3142,7 +3142,7 @@ bool FEngineLoop::LoadStartupCoreModules()
 #endif
 
 	FModuleManager::Get().LoadModule(TEXT("PacketHandler"));
-
+	FModuleManager::Get().LoadModule(TEXT("NetworkReplayStreaming"));
 
 	return bSuccess;
 }
diff --git a/Engine/Source/Runtime/NetworkReplayStreaming/NetworkReplayStreaming/Private/NetworkReplayStreaming.cpp b/Engine/Source/Runtime/NetworkReplayStreaming/NetworkReplayStreaming/Private/NetworkReplayStreaming.cpp
--- a/Engine/Source/Runtime/NetworkReplayStreaming/NetworkReplayStreaming/Private/NetworkReplayStreaming.cpp
+++ b/Engine/Source/Runtime/NetworkReplayStreaming/NetworkReplayStreaming/Private/NetworkReplayStreaming.cpp
@@ -99,4 +99,25 @@ void FNetworkReplayStreaming::Flush()
 			ReplayFactory.Flush();
 		}
 	}
-}
\ No newline at end of file
+}
+
+bool FNetworkReplayStreaming::Exec(UWorld* InWorld, const TCHAR* Cmd, FOutputDevice& Ar)
+{
+	// expected usage is "replaystreamer <streamer factory name> <streamer specific args>"
+	if (FParse::Command(&Cmd, TEXT("REPLAYSTREAMER")))
+	{
+		FString FactoryName = FParse::Token(Cmd, false);
+		if (!FactoryName.IsEmpty())
+		{
+			TSharedPtr<INetworkReplayStreamer> Streamer = GetFactory(*FactoryName).CreateReplayStreamer();
+			if (Streamer.IsValid())
+			{
+				Streamer->Exec(Cmd, Ar);
+			}
+		}
+
+		return true;
+	}
+
+	return false;
+}
diff --git a/Engine/Source/Runtime/NetworkReplayStreaming/NetworkReplayStreaming/Public/NetworkReplayStreaming.h b/Engine/Source/Runtime/NetworkReplayStreaming/NetworkReplayStreaming/Public/NetworkReplayStreaming.h
--- a/Engine/Source/Runtime/NetworkReplayStreaming/NetworkReplayStreaming/Public/NetworkReplayStreaming.h
+++ b/Engine/Source/Runtime/NetworkReplayStreaming/NetworkReplayStreaming/Public/NetworkReplayStreaming.h
@@ -606,23 +606,28 @@ public:
 	virtual EStreamingOperationResult GetDemoPath(FString& DemoPath) const = 0;
 
 	virtual void SetAnalyticsProvider(TSharedPtr<IAnalyticsProvider>& InProvider) {}
+
+	virtual void Exec(const TCHAR* Cmd, FOutputDevice& Ar) {}
 };
 
 /** Replay streamer factory */
 class INetworkReplayStreamingFactory : public IModuleInterface
 {
 public:
-	virtual TSharedPtr< INetworkReplayStreamer > CreateReplayStreamer() = 0;
+	virtual TSharedPtr<INetworkReplayStreamer> CreateReplayStreamer() = 0;
 	virtual void Flush() {}
 };
 
 /** Replay streaming factory manager */
-class FNetworkReplayStreaming : public IModuleInterface
+class FNetworkReplayStreaming : public IModuleInterface, public FSelfRegisteringExec
 {
 public:
+	FNetworkReplayStreaming() {}
+	virtual ~FNetworkReplayStreaming() {}
+
 	static inline FNetworkReplayStreaming& Get()
 	{
-		return FModuleManager::LoadModuleChecked< FNetworkReplayStreaming >( "NetworkReplayStreaming" );
+		return FModuleManager::LoadModuleChecked<FNetworkReplayStreaming>("NetworkReplayStreaming");
 	}
 
 	NETWORKREPLAYSTREAMING_API void Flush();
@@ -645,6 +650,9 @@ public:
 	// If false, it's up to the streamer to determine a proper postfix.
 	static NETWORKREPLAYSTREAMING_API bool UseDateTimeAsAutomaticReplayPostfix();
 
+	// FSelfRegisteringExec interface
+	virtual bool Exec(UWorld* InWorld, const TCHAR* Cmd, FOutputDevice& Ar) override;
+
 private:
 	TSet<FName> LoadedFactories;
 };
