Making some internal serialization values publicly accessible

diff --git a/Engine/Source/Runtime/NetworkReplayStreaming/LocalFileNetworkReplayStreaming/Private/LocalFileNetworkReplayStreaming.cpp b/Engine/Source/Runtime/NetworkReplayStreaming/LocalFileNetworkReplayStreaming/Private/LocalFileNetworkReplayStreaming.cpp
--- a/Engine/Source/Runtime/NetworkReplayStreaming/LocalFileNetworkReplayStreaming/Private/LocalFileNetworkReplayStreaming.cpp
+++ b/Engine/Source/Runtime/NetworkReplayStreaming/LocalFileNetworkReplayStreaming/Private/LocalFileNetworkReplayStreaming.cpp
@@ -58,15 +58,16 @@ namespace LocalFileReplay
 		HISTORY_LATEST 							= HISTORY_PLUS_ONE - 1
 	};
 
-	const uint32 FileMagic = 0x1CA2E27F;
-	const uint32 MaxFriendlyNameLen = 256;
-
 	TAutoConsoleVariable<int32> CVarMaxCacheSize(TEXT("localReplay.MaxCacheSize"), 1024 * 1024 * 10, TEXT(""));
 	TAutoConsoleVariable<int32> CVarMaxBufferedStreamChunks(TEXT("localReplay.MaxBufferedStreamChunks"), 5, TEXT(""));
 	TAutoConsoleVariable<int32> CVarAllowLiveStreamDelete(TEXT("localReplay.AllowLiveStreamDelete"), 1, TEXT(""));
 	TAutoConsoleVariable<float> CVarChunkUploadDelayInSeconds(TEXT("localReplay.ChunkUploadDelayInSeconds"), 20.0f, TEXT(""));
 };
 
+const uint32 FLocalFileNetworkReplayStreamer::FileMagic = 0x1CA2E27F;
+const uint32 FLocalFileNetworkReplayStreamer::MaxFriendlyNameLen = 256;
+const uint32 FLocalFileNetworkReplayStreamer::LatestVersion = LocalFileReplay::HISTORY_LATEST;
+
 FLocalFileNetworkReplayStreamer::FLocalFileSerializationInfo::FLocalFileSerializationInfo() 
 	: FileVersion(LocalFileReplay::HISTORY_LATEST)
 {
@@ -195,7 +196,7 @@ bool FLocalFileNetworkReplayStreamer::ReadReplayInfo(FArchive& Archive, FLocalFi
 		uint32 FileVersion;
 		Archive << FileVersion;
 
-		if (MagicNumber == LocalFileReplay::FileMagic)
+		if (MagicNumber == FLocalFileNetworkReplayStreamer::FileMagic)
 		{
 			SerializationInfo.FileVersion = FileVersion;
 
@@ -505,7 +506,7 @@ bool FLocalFileNetworkReplayStreamer::WriteReplayInfo(FArchive& Archive, const F
 
 	Archive.Seek(0);
 
-	uint32 MagicNumber = LocalFileReplay::FileMagic;
+	uint32 MagicNumber = FLocalFileNetworkReplayStreamer::FileMagic;
 	Archive << MagicNumber;
 
 	Archive << SerializationInfo.FileVersion;
@@ -2397,7 +2398,7 @@ const FString& FLocalFileNetworkReplayStreamer::GetDefaultDemoSavePath()
 
 uint32 FLocalFileNetworkReplayStreamer::GetMaxFriendlyNameSize() const
 {
-	return LocalFileReplay::MaxFriendlyNameLen;
+	return FLocalFileNetworkReplayStreamer::MaxFriendlyNameLen;
 }
 
 void FLocalFileNetworkReplayStreamer::DownloadHeader(const FDownloadHeaderCallback& Delegate)
diff --git a/Engine/Source/Runtime/NetworkReplayStreaming/LocalFileNetworkReplayStreaming/Public/LocalFileNetworkReplayStreaming.h b/Engine/Source/Runtime/NetworkReplayStreaming/LocalFileNetworkReplayStreaming/Public/LocalFileNetworkReplayStreaming.h
--- a/Engine/Source/Runtime/NetworkReplayStreaming/LocalFileNetworkReplayStreaming/Public/LocalFileNetworkReplayStreaming.h
+++ b/Engine/Source/Runtime/NetworkReplayStreaming/LocalFileNetworkReplayStreaming/Public/LocalFileNetworkReplayStreaming.h
@@ -637,6 +637,10 @@ protected:
 
 public:
 	static const FString& GetDefaultDemoSavePath();
+
+	static const uint32 FileMagic;
+	static const uint32 MaxFriendlyNameLen;
+	static const uint32 LatestVersion;
 };
 
 class LOCALFILENETWORKREPLAYSTREAMING_API FLocalFileNetworkReplayStreamingFactory : public INetworkReplayStreamingFactory, public FTickableGameObject
