Bump default max buffered stream chunks count

#rnx

diff --git a/Engine/Source/Runtime/NetworkReplayStreaming/LocalFileNetworkReplayStreaming/Private/LocalFileNetworkReplayStreaming.cpp b/Engine/Source/Runtime/NetworkReplayStreaming/LocalFileNetworkReplayStreaming/Private/LocalFileNetworkReplayStreaming.cpp
--- a/Engine/Source/Runtime/NetworkReplayStreaming/LocalFileNetworkReplayStreaming/Private/LocalFileNetworkReplayStreaming.cpp
+++ b/Engine/Source/Runtime/NetworkReplayStreaming/LocalFileNetworkReplayStreaming/Private/LocalFileNetworkReplayStreaming.cpp
@@ -59,7 +59,7 @@ namespace LocalFileReplay
 	};
 
 	TAutoConsoleVariable<int32> CVarMaxCacheSize(TEXT("localReplay.MaxCacheSize"), 1024 * 1024 * 10, TEXT(""));
-	TAutoConsoleVariable<int32> CVarMaxBufferedStreamChunks(TEXT("localReplay.MaxBufferedStreamChunks"), 5, TEXT(""));
+	TAutoConsoleVariable<int32> CVarMaxBufferedStreamChunks(TEXT("localReplay.MaxBufferedStreamChunks"), 10, TEXT(""));
 	TAutoConsoleVariable<int32> CVarAllowLiveStreamDelete(TEXT("localReplay.AllowLiveStreamDelete"), 1, TEXT(""));
 	TAutoConsoleVariable<float> CVarChunkUploadDelayInSeconds(TEXT("localReplay.ChunkUploadDelayInSeconds"), 20.0f, TEXT(""));
 };
