Add UpdatePlaybackTime to replay streamer interface

diff --git a/Engine/Source/Runtime/Engine/Private/DemoNetDriver.cpp b/Engine/Source/Runtime/Engine/Private/DemoNetDriver.cpp
--- a/Engine/Source/Runtime/Engine/Private/DemoNetDriver.cpp
+++ b/Engine/Source/Runtime/Engine/Private/DemoNetDriver.cpp
@@ -3860,6 +3860,8 @@ void UDemoNetDriver::TickDemoPlayback(float DeltaSeconds)
 	// Clamp time
 	DemoCurrentTime = FMath::Clamp(DemoCurrentTime, 0.0f, DemoTotalTime + 0.01f);
 
+	ReplayStreamer->UpdatePlaybackTime(GetDemoCurrentTimeInMS());
+
 	bool bProcessAvailableData = (PlaybackPackets.Num() > 0) || ReplayStreamer->IsDataAvailable();
 	
 	if (CVarFastForwardLevelsPausePlayback.GetValueOnAnyThread() == 0)
diff --git a/Engine/Source/Runtime/NetworkReplayStreaming/HttpNetworkReplayStreaming/Public/HttpNetworkReplayStreaming.h b/Engine/Source/Runtime/NetworkReplayStreaming/HttpNetworkReplayStreaming/Public/HttpNetworkReplayStreaming.h
--- a/Engine/Source/Runtime/NetworkReplayStreaming/HttpNetworkReplayStreaming/Public/HttpNetworkReplayStreaming.h
+++ b/Engine/Source/Runtime/NetworkReplayStreaming/HttpNetworkReplayStreaming/Public/HttpNetworkReplayStreaming.h
@@ -215,6 +215,7 @@ public:
 	virtual void		GotoCheckpointIndex( const int32 CheckpointIndex, const FGotoCallback& Delegate, EReplayCheckpointType CheckpointType ) override;
 	virtual void		GotoTimeInMS( const uint32 TimeInMS, const FGotoCallback& Delegate, EReplayCheckpointType CheckpointType ) override;
 	virtual void		UpdateTotalDemoTime( uint32 TimeInMS ) override;
+	virtual void		UpdatePlaybackTime(uint32 TimeInMS) override {}
 	virtual uint32		GetTotalDemoTime() const override { return TotalDemoTimeInMS; }
 	virtual bool		IsDataAvailable() const override;
 	virtual void		SetHighPriorityTimeRange( const uint32 StartTimeInMS, const uint32 EndTimeInMS ) override;
diff --git a/Engine/Source/Runtime/NetworkReplayStreaming/InMemoryNetworkReplayStreaming/Public/InMemoryNetworkReplayStreaming.h b/Engine/Source/Runtime/NetworkReplayStreaming/InMemoryNetworkReplayStreaming/Public/InMemoryNetworkReplayStreaming.h
--- a/Engine/Source/Runtime/NetworkReplayStreaming/InMemoryNetworkReplayStreaming/Public/InMemoryNetworkReplayStreaming.h
+++ b/Engine/Source/Runtime/NetworkReplayStreaming/InMemoryNetworkReplayStreaming/Public/InMemoryNetworkReplayStreaming.h
@@ -133,6 +133,7 @@ public:
 	virtual void GotoCheckpointIndex( const int32 CheckpointIndex, const FGotoCallback& Delegate, EReplayCheckpointType CheckpointType ) override;
 	virtual void GotoTimeInMS( const uint32 TimeInMS, const FGotoCallback& Delegate, EReplayCheckpointType CheckpointType ) override;
 	virtual void UpdateTotalDemoTime( uint32 TimeInMS ) override;
+	virtual void UpdatePlaybackTime(uint32 TimeInMS) override {}
 	virtual uint32 GetTotalDemoTime() const override;
 	virtual bool IsDataAvailable() const override;
 	virtual void SetHighPriorityTimeRange( const uint32 StartTimeInMS, const uint32 EndTimeInMS ) override { }
diff --git a/Engine/Source/Runtime/NetworkReplayStreaming/LocalFileNetworkReplayStreaming/Public/LocalFileNetworkReplayStreaming.h b/Engine/Source/Runtime/NetworkReplayStreaming/LocalFileNetworkReplayStreaming/Public/LocalFileNetworkReplayStreaming.h
--- a/Engine/Source/Runtime/NetworkReplayStreaming/LocalFileNetworkReplayStreaming/Public/LocalFileNetworkReplayStreaming.h
+++ b/Engine/Source/Runtime/NetworkReplayStreaming/LocalFileNetworkReplayStreaming/Public/LocalFileNetworkReplayStreaming.h
@@ -408,6 +408,7 @@ public:
 	virtual void GotoCheckpointIndex(const int32 CheckpointIndex, const FGotoCallback& Delegate, EReplayCheckpointType CheckpointType) override;
 	virtual void GotoTimeInMS(const uint32 TimeInMS, const FGotoCallback& Delegate, EReplayCheckpointType CheckpointType) override;
 	virtual void UpdateTotalDemoTime(uint32 TimeInMS) override;
+	virtual void UpdatePlaybackTime(uint32 TimeInMS) override {}
 	virtual uint32 GetTotalDemoTime() const override { return CurrentReplayInfo.LengthInMS; }
 	virtual bool IsDataAvailable() const override;
 	virtual void SetHighPriorityTimeRange(const uint32 StartTimeInMS, const uint32 EndTimeInMS) override;
diff --git a/Engine/Source/Runtime/NetworkReplayStreaming/NetworkReplayStreaming/Public/NetworkReplayStreaming.h b/Engine/Source/Runtime/NetworkReplayStreaming/NetworkReplayStreaming/Public/NetworkReplayStreaming.h
--- a/Engine/Source/Runtime/NetworkReplayStreaming/NetworkReplayStreaming/Public/NetworkReplayStreaming.h
+++ b/Engine/Source/Runtime/NetworkReplayStreaming/NetworkReplayStreaming/Public/NetworkReplayStreaming.h
@@ -485,6 +485,8 @@ public:
 	virtual bool IsCheckpointTypeSupported(EReplayCheckpointType CheckpointType) const = 0;
 
 	virtual void UpdateTotalDemoTime(uint32 TimeInMS) = 0;
+	virtual void UpdatePlaybackTime(uint32 TimeInMS) = 0;
+
 	virtual uint32 GetTotalDemoTime() const = 0;
 	virtual bool IsDataAvailable() const = 0;
 	virtual void SetHighPriorityTimeRange(const uint32 StartTimeInMS, const uint32 EndTimeInMS) = 0;
diff --git a/Engine/Source/Runtime/NetworkReplayStreaming/NullNetworkReplayStreaming/Public/NullNetworkReplayStreaming.h b/Engine/Source/Runtime/NetworkReplayStreaming/NullNetworkReplayStreaming/Public/NullNetworkReplayStreaming.h
--- a/Engine/Source/Runtime/NetworkReplayStreaming/NullNetworkReplayStreaming/Public/NullNetworkReplayStreaming.h
+++ b/Engine/Source/Runtime/NetworkReplayStreaming/NullNetworkReplayStreaming/Public/NullNetworkReplayStreaming.h
@@ -54,6 +54,7 @@ public:
 	virtual void GotoCheckpointIndex( const int32 CheckpointIndex, const FGotoCallback& Delegate, EReplayCheckpointType CheckpointType ) override;
 	virtual void GotoTimeInMS( const uint32 TimeInMS, const FGotoCallback& Delegate, EReplayCheckpointType CheckpointType ) override;
 	virtual void UpdateTotalDemoTime( uint32 TimeInMS ) override;
+	virtual void UpdatePlaybackTime(uint32 TimeInMS) override {}
 	virtual uint32 GetTotalDemoTime() const override { return ReplayInfo.LengthInMS; }
 	virtual bool IsDataAvailable() const override;
 	virtual void SetHighPriorityTimeRange( const uint32 StartTimeInMS, const uint32 EndTimeInMS ) override { }
