Support using -replaystreamer= or -replaystreameroverride= on the command line for testing

diff --git a/Engine/Source/Runtime/NetworkReplayStreaming/NetworkReplayStreaming/Private/NetworkReplayStreaming.cpp b/Engine/Source/Runtime/NetworkReplayStreaming/NetworkReplayStreaming/Private/NetworkReplayStreaming.cpp
--- a/Engine/Source/Runtime/NetworkReplayStreaming/NetworkReplayStreaming/Private/NetworkReplayStreaming.cpp
+++ b/Engine/Source/Runtime/NetworkReplayStreaming/NetworkReplayStreaming/Private/NetworkReplayStreaming.cpp
@@ -21,7 +21,7 @@ INetworkReplayStreamingFactory& FNetworkReplayStreaming::GetFactory(const TCHAR*
 	}
 
 	FString CmdlineFactoryName;
-	if (FParse::Value(FCommandLine::Get(), TEXT("-REPLAYSTREAMER="), CmdlineFactoryName))
+	if (FParse::Value(FCommandLine::Get(), TEXT("-REPLAYSTREAMER="), CmdlineFactoryName) || FParse::Value(FCommandLine::Get(), TEXT("-REPLAYSTREAMEROVERRIDE="), CmdlineFactoryName))
 	{
 		FactoryName = CmdlineFactoryName;
 	}
