Change default replay streamer from null to local file

diff --git a/Engine/Source/Runtime/Engine/Engine.Build.cs b/Engine/Source/Runtime/Engine/Engine.Build.cs
--- a/Engine/Source/Runtime/Engine/Engine.Build.cs
+++ b/Engine/Source/Runtime/Engine/Engine.Build.cs
@@ -303,6 +303,7 @@ public class Engine : ModuleRules
 		DynamicallyLoadedModuleNames.AddRange(
 			new string[] {
 				"NullNetworkReplayStreaming",
+				"LocalFileNetworkReplayStreaming",
 				"HttpNetworkReplayStreaming",
 				"Advertising"
 			}
diff --git a/Engine/Source/Runtime/NetworkReplayStreaming/NetworkReplayStreaming/Private/NetworkReplayStreaming.cpp b/Engine/Source/Runtime/NetworkReplayStreaming/NetworkReplayStreaming/Private/NetworkReplayStreaming.cpp
--- a/Engine/Source/Runtime/NetworkReplayStreaming/NetworkReplayStreaming/Private/NetworkReplayStreaming.cpp
+++ b/Engine/Source/Runtime/NetworkReplayStreaming/NetworkReplayStreaming/Private/NetworkReplayStreaming.cpp
@@ -9,7 +9,9 @@ IMPLEMENT_MODULE( FNetworkReplayStreaming, NetworkReplayStreaming );
 
 INetworkReplayStreamingFactory& FNetworkReplayStreaming::GetFactory(const TCHAR* FactoryNameOverride)
 {
-	FString FactoryName = TEXT("NullNetworkReplayStreaming");
+	static const FString DefaultFactoryName = TEXT("LocalFileNetworkReplayStreaming");
+
+	FString FactoryName = DefaultFactoryName;
 
 	if (FactoryNameOverride == nullptr)
 	{
@@ -33,7 +35,7 @@ INetworkReplayStreamingFactory& FNetworkReplayStreaming::GetFactory(const TCHAR*
 	
 		if (!FModuleManager::Get().IsModuleLoaded(*FactoryName))
 		{
-			FactoryName = TEXT("NullNetworkReplayStreaming");
+			FactoryName = DefaultFactoryName;
 		}
 	}
 
