Replay interface changes to pass the demo url into the streamer

diff --git a/Engine/Source/Runtime/Engine/Private/DemoNetDriver.cpp b/Engine/Source/Runtime/Engine/Private/DemoNetDriver.cpp
--- a/Engine/Source/Runtime/Engine/Private/DemoNetDriver.cpp
+++ b/Engine/Source/Runtime/Engine/Private/DemoNetDriver.cpp
@@ -864,18 +864,6 @@ bool UDemoNetDriver::InitConnect(FNetworkNotify* InNotify, const FURL& ConnectUR
 	ServerConnection = NewObject<UNetConnection>(GetTransientPackage(), UDemoNetConnection::StaticClass());
 	ServerConnection->InitConnection(this, USOCK_Pending, ConnectURL, 1000000);
 
-	TArray<FString> UserNames;
-
-	if (GetWorld()->GetGameInstance()->GetFirstGamePlayer() != nullptr)
-	{
-		FUniqueNetIdRepl ViewerId = GetWorld()->GetGameInstance()->GetFirstGamePlayer()->GetPreferredUniqueNetId();
-
-		if (ViewerId.IsValid())
-		{ 
-			UserNames.Add(ViewerId.ToString());
-		}
-	}
-
 	const TCHAR* const LevelPrefixOverrideOption = DemoURL.GetOption(TEXT("LevelPrefixOverride="), nullptr);
 	if (LevelPrefixOverrideOption)
 	{
@@ -905,13 +893,24 @@ bool UDemoNetDriver::InitConnect(FNetworkNotify* InNotify, const FURL& ConnectUR
 	bWasStartStreamingSuccessful = true;
 
 	ActiveReplayName = DemoURL.Map;
-	ReplayStreamer->StartStreaming( 
-		DemoURL.Map, 
-		FString(),		// Friendly name isn't important for loading an existing replay.
-		UserNames, 
-		false, 
-		FNetworkVersion::GetReplayVersion(), 
-		FStartStreamingCallback::CreateUObject(this, &UDemoNetDriver::ReplayStreamingReady));
+
+	TArray<int32> UserIndices;
+	for (FLocalPlayerIterator It(GEngine, World); It; ++It)
+	{
+		if (*It)
+		{
+			UserIndices.Add(It->GetControllerId());
+		}
+	}
+	
+	FStartStreamingParameters Params;
+	Params.CustomName = DemoURL.Map;
+	Params.DemoURL = DemoURL;
+	Params.UserIndices = MoveTemp(UserIndices);
+	Params.bRecord = false;
+	Params.ReplayVersion = FNetworkVersion::GetReplayVersion();
+
+	ReplayStreamer->StartStreaming(Params, FStartStreamingCallback::CreateUObject(this, &UDemoNetDriver::ReplayStreamingReady));
 
 	return bWasStartStreamingSuccessful;
 }
@@ -1078,32 +1077,28 @@ bool UDemoNetDriver::InitListen(FNetworkNotify* InNotify, FURL& ListenURL, bool
 
 	bRecordMapChanges = ListenURL.GetOption(TEXT("RecordMapChanges"), nullptr) != nullptr;
 
-	TArray<FString> UserNames;
-	AGameStateBase* GameState = GetWorld()->GetGameState();
-
-	// If a client is recording a replay, GameState may not have replicated yet
-	if (GameState != nullptr)
+	TArray<int32> UserIndices;
+	for (FLocalPlayerIterator It(GEngine, World); It; ++It)
 	{
-		for (int32 i = 0; i < GameState->PlayerArray.Num(); i++)
+		if (*It)
 		{
-			APlayerState* PlayerState = GameState->PlayerArray[i];
-			if (PlayerState && !PlayerState->bIsABot && !PlayerState->bIsSpectator)
-			{
-				UserNames.Add(PlayerState->UniqueId.ToString());
-			}
+			UserIndices.Add(It->GetControllerId());
 		}
 	}
 
 	bIsWaitingForStream = true;
 
 	ActiveReplayName = DemoURL.Map;
-	ReplayStreamer->StartStreaming(
-		DemoURL.Map,
-		FriendlyNameOption != nullptr ? FString(FriendlyNameOption) : World->GetMapName(),
-		UserNames,
-		true,
-		FNetworkVersion::GetReplayVersion(),
-		FStartStreamingCallback::CreateUObject(this, &UDemoNetDriver::ReplayStreamingReady));
+
+	FStartStreamingParameters Params;
+	Params.CustomName = DemoURL.Map;
+	Params.FriendlyName = FriendlyNameOption != nullptr ? FString(FriendlyNameOption) : World->GetMapName();
+	Params.DemoURL = DemoURL;
+	Params.UserIndices = MoveTemp(UserIndices);
+	Params.bRecord = true;
+	Params.ReplayVersion = FNetworkVersion::GetReplayVersion();
+
+	ReplayStreamer->StartStreaming(Params, FStartStreamingCallback::CreateUObject(this, &UDemoNetDriver::ReplayStreamingReady));
 
 	AddNewLevel(World->GetOuter()->GetName());
 
@@ -1606,7 +1601,7 @@ void UDemoNetDriver::StopDemo()
 	{
 		UE_LOG(LogDemo, Log, TEXT("StopDemo: No demo is playing"));
 		ClearReplayTasks();
-		ActiveReplayName = FString();
+		ActiveReplayName.Empty();
 		ResetDemoState();
 		return;
 	}
@@ -1633,7 +1628,7 @@ void UDemoNetDriver::StopDemo()
 
 	ReplayStreamer->StopStreaming();
 	ClearReplayTasks();
-	ActiveReplayName = FString();
+	ActiveReplayName.Empty();
 	ResetDemoState();
 
 	check(!IsRecording() && !IsPlaying());
diff --git a/Engine/Source/Runtime/NetworkReplayStreaming/HttpNetworkReplayStreaming/Private/HttpNetworkReplayStreaming.cpp b/Engine/Source/Runtime/NetworkReplayStreaming/HttpNetworkReplayStreaming/Private/HttpNetworkReplayStreaming.cpp
--- a/Engine/Source/Runtime/NetworkReplayStreaming/HttpNetworkReplayStreaming/Private/HttpNetworkReplayStreaming.cpp
+++ b/Engine/Source/Runtime/NetworkReplayStreaming/HttpNetworkReplayStreaming/Private/HttpNetworkReplayStreaming.cpp
@@ -227,21 +227,6 @@ FHttpNetworkReplayStreamer::FHttpNetworkReplayStreamer() :
 	GConfig->GetString( TEXT( "HttpNetworkReplayStreaming" ), TEXT( "ServerURL" ), ServerURL, GEngineIni );
 }
 
-void FHttpNetworkReplayStreamer::StartStreaming(const FString& CustomName, const FString& FriendlyName, const TArray< int32 >& UserIndices, bool bRecord, const FNetworkReplayVersion& InReplayVersion, const FStartStreamingCallback& Delegate)
-{
-	TArray<FString> UserStrings;
-	for (int32 UserIndex : UserIndices)
-	{
-		FString UserString(GetUserStringFromUserIndex(UserIndex));
-		if (!UserString.IsEmpty())
-		{
-			UserStrings.Add(MoveTemp(UserString));
-		}
-	}
-
-	StartStreaming(CustomName, FriendlyName, UserStrings, bRecord, InReplayVersion, Delegate);
-}
-
 FString FHttpNetworkReplayStreamer::GetRecordingMetadata() const
 {
 	FString MetaString;
@@ -249,10 +234,10 @@ FString FHttpNetworkReplayStreamer::GetRecordingMetadata() const
 	return MetaString;
 }
 
-void FHttpNetworkReplayStreamer::StartStreaming( const FString& CustomName, const FString& FriendlyName, const TArray< FString >& UserNames, bool bRecord, const FNetworkReplayVersion& InReplayVersion, const FStartStreamingCallback& Delegate )
+void FHttpNetworkReplayStreamer::StartStreaming(const FStartStreamingParameters& Params, const FStartStreamingCallback& Delegate)
 {
 	FStartStreamingResult StreamingResult;
-	StreamingResult.bRecording = bRecord;
+	StreamingResult.bRecording = Params.bRecord;
 
 	if ( !SessionName.IsEmpty() )
 	{
@@ -282,21 +267,21 @@ void FHttpNetworkReplayStreamer::StartStreaming( const FString& CustomName, cons
 		return;
 	}
 
-	ReplayVersion = InReplayVersion;
+	ReplayVersion = Params.ReplayVersion;
 
 	// Remember the delegate, which we'll call as soon as the header is available
 	StartStreamingDelegate = Delegate;
 
 	// Setup the archives
-	StreamArchive.SetIsLoading(!bRecord);
-	StreamArchive.SetIsSaving(bRecord);
+	StreamArchive.SetIsLoading(!Params.bRecord);
+	StreamArchive.SetIsSaving(Params.bRecord);
 	StreamArchive.bAtEndOfReplay	= false;
 
-	HeaderArchive.SetIsLoading(!bRecord);
-	HeaderArchive.SetIsSaving(bRecord);
+	HeaderArchive.SetIsLoading(!Params.bRecord);
+	HeaderArchive.SetIsSaving(Params.bRecord);
 
-	CheckpointArchive.SetIsLoading(!bRecord);
-	CheckpointArchive.SetIsSaving(bRecord);
+	CheckpointArchive.SetIsLoading(!Params.bRecord);
+	CheckpointArchive.SetIsSaving(Params.bRecord);
 
 	LastChunkTime = FPlatformTime::Seconds();
 
@@ -316,18 +301,18 @@ void FHttpNetworkReplayStreamer::StartStreaming( const FString& CustomName, cons
 
 	RefreshViewerFails = 0;
 
-	if ( !bRecord )
+	if ( !Params.bRecord )
 	{
 		// We are streaming down
 		StreamerState = EStreamerState::StreamingDown;
 
-		SessionName = CustomName;
+		SessionName = Params.CustomName;
 
 		FString UserName;
 
-		if ( UserNames.Num() == 1 )
+		if (Params.UserIndices.Num() == 1)
 		{
-			UserName = UserNames[0];
+			UserName = GetUserStringFromUserIndex(Params.UserIndices[0]);
 		}
 
 		// Notify the http server that we want to start downloading a replay
@@ -358,15 +343,15 @@ void FHttpNetworkReplayStreamer::StartStreaming( const FString& CustomName, cons
 
 		FString URL;
 
-		if ( !CustomName.IsEmpty() )
+		if ( !Params.CustomName.IsEmpty() )
 		{
-			const FString SessionIDOverride = CustomName.ToLower();
+			const FString SessionIDOverride = Params.CustomName.ToLower();
 
-			URL = FString::Printf( TEXT( "%sreplay/%s?app=%s&version=%u&cl=%u&friendlyName=%s" ), *ServerURL, *SessionIDOverride, *ReplayVersion.AppString, ReplayVersion.NetworkVersion, ReplayVersion.Changelist, *FGenericPlatformHttp::UrlEncode( FriendlyName ) );
+			URL = FString::Printf( TEXT( "%sreplay/%s?app=%s&version=%u&cl=%u&friendlyName=%s" ), *ServerURL, *SessionIDOverride, *ReplayVersion.AppString, ReplayVersion.NetworkVersion, ReplayVersion.Changelist, *FGenericPlatformHttp::UrlEncode( Params.FriendlyName ) );
 		}
 		else
 		{
-			URL = FString::Printf( TEXT( "%sreplay?app=%s&version=%u&cl=%u&friendlyName=%s" ), *ServerURL, *ReplayVersion.AppString, ReplayVersion.NetworkVersion, ReplayVersion.Changelist, *FGenericPlatformHttp::UrlEncode( FriendlyName ) );
+			URL = FString::Printf( TEXT( "%sreplay?app=%s&version=%u&cl=%u&friendlyName=%s" ), *ServerURL, *ReplayVersion.AppString, ReplayVersion.NetworkVersion, ReplayVersion.Changelist, *FGenericPlatformHttp::UrlEncode( Params.FriendlyName ) );
 		}
 
 		FString MetaString = GetRecordingMetadata();
@@ -384,11 +369,15 @@ void FHttpNetworkReplayStreamer::StartStreaming( const FString& CustomName, cons
 
 		HttpRequest->SetHeader( TEXT( "Content-Type" ), TEXT( "application/json" ) );
 
-		if ( UserNames.Num() > 0 )
+		if (Params.UserIndices.Num() > 0)
 		{
 			FNetworkReplayUserList UserList;
 
-			UserList.Users = UserNames;
+			for (int32 UserIdx : Params.UserIndices)
+			{
+				UserList.Users.Add(GetUserStringFromUserIndex(UserIdx));
+			}
+
 			HttpRequest->SetContentAsString( UserList.ToJson() );
 		}
 
@@ -1527,16 +1516,6 @@ void FHttpNetworkReplayStreamer::DeleteFinishedStream( const FString& StreamName
 }
 
 void FHttpNetworkReplayStreamer::EnumerateStreams( const FNetworkReplayVersion& InReplayVersion, const int32 UserIndex, const FString& MetaString, const TArray< FString >& ExtraParms, const FEnumerateStreamsCallback& Delegate )
-{
-	EnumerateStreams( InReplayVersion, GetUserStringFromUserIndex(UserIndex), MetaString, ExtraParms, Delegate );
-}
-
-void FHttpNetworkReplayStreamer::EnumerateStreams( const FNetworkReplayVersion& InReplayVersion, const FString& UserString, const FString& MetaString, const FEnumerateStreamsCallback& Delegate )
-{
-	EnumerateStreams( InReplayVersion, UserString, MetaString, TArray< FString >(), Delegate );
-}
-
-void FHttpNetworkReplayStreamer::EnumerateStreams( const FNetworkReplayVersion& InReplayVersion, const FString& UserString, const FString& MetaString, const TArray< FString >& ExtraParms, const FEnumerateStreamsCallback& Delegate )
 {
 	if ( ServerURL.IsEmpty() )
 	{
@@ -1570,6 +1549,7 @@ void FHttpNetworkReplayStreamer::EnumerateStreams( const FNetworkReplayVersion&
 	}
 
 	// Add optional User parameter (filter replays by a user that was in the replay)
+	FString UserString = GetUserStringFromUserIndex(UserIndex);
 	if ( !UserString.IsEmpty() )
 	{
 		URL += FString::Printf( TEXT( "&user=%s" ), *FGenericPlatformHttp::UrlEncode( UserString ) );
@@ -1591,11 +1571,6 @@ void FHttpNetworkReplayStreamer::EnumerateStreams( const FNetworkReplayVersion&
 }
 
 void FHttpNetworkReplayStreamer::EnumerateRecentStreams( const FNetworkReplayVersion& InReplayVersion, const int32 UserIndex, const FEnumerateStreamsCallback& Delegate )
-{
-	EnumerateRecentStreams( InReplayVersion, GetUserStringFromUserIndex(UserIndex), Delegate );
-}
-
-void FHttpNetworkReplayStreamer::EnumerateRecentStreams( const FNetworkReplayVersion& InReplayVersion, const FString& InRecentViewer, const FEnumerateStreamsCallback& Delegate )
 {
 	if ( ServerURL.IsEmpty() )
 	{
@@ -1607,7 +1582,7 @@ void FHttpNetworkReplayStreamer::EnumerateRecentStreams( const FNetworkReplayVer
 	TSharedRef<IHttpRequest> HttpRequest = FHttpModule::Get().CreateRequest();
 
 	// Enumerate all of the sessions
-	HttpRequest->SetURL( FString::Printf( TEXT( "%sreplay?app=%s&version=%u&cl=%u&recent=%s" ), *ServerURL, *InReplayVersion.AppString, InReplayVersion.NetworkVersion, InReplayVersion.Changelist, *InRecentViewer ) );
+	HttpRequest->SetURL( FString::Printf( TEXT( "%sreplay?app=%s&version=%u&cl=%u&recent=%s" ), *ServerURL, *InReplayVersion.AppString, InReplayVersion.NetworkVersion, InReplayVersion.Changelist, *GetUserStringFromUserIndex(UserIndex) ) );
 	HttpRequest->SetVerb( TEXT( "GET" ) );
 
 	HttpRequest->OnProcessRequestComplete().BindRaw( this, &FHttpNetworkReplayStreamer::HttpEnumerateSessionsFinished, Delegate );
diff --git a/Engine/Source/Runtime/NetworkReplayStreaming/HttpNetworkReplayStreaming/Public/HttpNetworkReplayStreaming.h b/Engine/Source/Runtime/NetworkReplayStreaming/HttpNetworkReplayStreaming/Public/HttpNetworkReplayStreaming.h
--- a/Engine/Source/Runtime/NetworkReplayStreaming/HttpNetworkReplayStreaming/Public/HttpNetworkReplayStreaming.h
+++ b/Engine/Source/Runtime/NetworkReplayStreaming/HttpNetworkReplayStreaming/Public/HttpNetworkReplayStreaming.h
@@ -3,7 +3,6 @@
 
 #include "CoreMinimal.h"
 #include "Stats/Stats.h"
-#include "Misc/NetworkVersion.h"
 #include "NetworkReplayStreaming.h"
 #include "Interfaces/IHttpRequest.h"
 #include "Tickable.h"
@@ -205,8 +204,7 @@ public:
 	FHttpNetworkReplayStreamer();
 
 	/** INetworkReplayStreamer implementation */
-	virtual void		StartStreaming( const FString& CustomName, const FString& FriendlyName, const TArray< FString >& UserNames, bool bRecord, const FNetworkReplayVersion& ReplayVersion, const FStartStreamingCallback& Delegate ) override;
-	virtual void		StartStreaming( const FString& CustomName, const FString& FriendlyName, const TArray< int32 >& UserIndices, bool bRecord, const FNetworkReplayVersion& ReplayVersion, const FStartStreamingCallback& Delegate ) override;
+	virtual void		StartStreaming(const FStartStreamingParameters& Params, const FStartStreamingCallback& Delegate) override;
 	virtual void		StopStreaming() override;
 	virtual FArchive*	GetHeaderArchive() override;
 	virtual FArchive*	GetStreamingArchive() override;
@@ -224,13 +222,10 @@ public:
 	virtual bool		IsLive() const override;
 	virtual void		DeleteFinishedStream( const FString& StreamName, const FDeleteFinishedStreamCallback& Delegate ) override;
 	virtual void		DeleteFinishedStream( const FString& StreamName, const int32 UserIndex, const FDeleteFinishedStreamCallback& Delegate ) override;
-	virtual void		EnumerateStreams( const FNetworkReplayVersion& ReplayVersion, const FString& UserString, const FString& MetaString, const FEnumerateStreamsCallback& Delegate ) override;
-	virtual void		EnumerateStreams( const FNetworkReplayVersion& InReplayVersion, const FString& UserString, const FString& MetaString, const TArray< FString >& ExtraParms, const FEnumerateStreamsCallback& Delegate ) override;
 	virtual void		EnumerateStreams( const FNetworkReplayVersion& InReplayVersion, const int32 UserIndex, const FString& MetaString, const TArray< FString >& ExtraParms, const FEnumerateStreamsCallback& Delegate ) override;
 	virtual void		EnumerateEvents( const FString& Group, const FEnumerateEventsCallback& Delegate ) override;
 	virtual void		EnumerateEvents( const FString& ReplayName, const FString& Group, const FEnumerateEventsCallback& Delegate ) override;
 	virtual void		EnumerateEvents( const FString& ReplayName, const FString& Group, const int32 UserIndex, const FEnumerateEventsCallback& Delegate ) override;
-	virtual void		EnumerateRecentStreams( const FNetworkReplayVersion& ReplayVersion, const FString& RecentViewer, const FEnumerateStreamsCallback& Delegate ) override;
 	virtual void		EnumerateRecentStreams( const FNetworkReplayVersion& ReplayVersion, const int32 UserIndex, const FEnumerateStreamsCallback& Delegate ) override;
 	virtual void		AddUserToReplay(const FString& UserString) override;
 	virtual void		RequestEventData(const FString& EventId, const FRequestEventDataCallback& Delegate) override;
diff --git a/Engine/Source/Runtime/NetworkReplayStreaming/InMemoryNetworkReplayStreaming/Private/InMemoryNetworkReplayStreaming.cpp b/Engine/Source/Runtime/NetworkReplayStreaming/InMemoryNetworkReplayStreaming/Private/InMemoryNetworkReplayStreaming.cpp
--- a/Engine/Source/Runtime/NetworkReplayStreaming/InMemoryNetworkReplayStreaming/Private/InMemoryNetworkReplayStreaming.cpp
+++ b/Engine/Source/Runtime/NetworkReplayStreaming/InMemoryNetworkReplayStreaming/Private/InMemoryNetworkReplayStreaming.cpp
@@ -5,7 +5,6 @@
 #include "Misc/Guid.h"
 #include "Misc/DateTime.h"
 #include "HAL/IConsoleManager.h"
-#include "Misc/NetworkVersion.h"
 
 DEFINE_LOG_CATEGORY_STATIC( LogMemoryReplay, Log, All );
 
diff --git a/Engine/Source/Runtime/NetworkReplayStreaming/InMemoryNetworkReplayStreaming/Public/InMemoryNetworkReplayStreaming.h b/Engine/Source/Runtime/NetworkReplayStreaming/InMemoryNetworkReplayStreaming/Public/InMemoryNetworkReplayStreaming.h
--- a/Engine/Source/Runtime/NetworkReplayStreaming/InMemoryNetworkReplayStreaming/Public/InMemoryNetworkReplayStreaming.h
+++ b/Engine/Source/Runtime/NetworkReplayStreaming/InMemoryNetworkReplayStreaming/Public/InMemoryNetworkReplayStreaming.h
@@ -123,8 +123,7 @@ public:
 	}
 	
 	/** INetworkReplayStreamer implementation */
-	virtual void StartStreaming( const FString& CustomName, const FString& FriendlyName, const TArray< FString >& UserNames, bool bRecord, const FNetworkReplayVersion& ReplayVersion, const FStartStreamingCallback& Delegate ) override;
-	virtual void StartStreaming( const FString& CustomName, const FString& FriendlyName, const TArray< int32 >& UserIndices, bool bRecord, const FNetworkReplayVersion& ReplayVersion, const FStartStreamingCallback& Delegate ) override;
+	virtual void StartStreaming(const FStartStreamingParameters& Params, const FStartStreamingCallback& Delegate) override;
 	virtual void StopStreaming() override;
 	virtual FArchive* GetHeaderArchive() override;
 	virtual FArchive* GetStreamingArchive() override;
@@ -167,7 +166,7 @@ public:
 	virtual void RenameReplayFriendlyName(const FString& ReplayName, const FString& NewFriendlyName, const int32 UserIndex, const FRenameReplayCallback& Delegate) override;
 	virtual void RenameReplay(const FString& ReplayName, const FString& NewName, const FRenameReplayCallback& Delegate) override;
 	virtual void RenameReplay(const FString& ReplayName, const FString& NewName, const int32 UserIndex, const FRenameReplayCallback& Delegate) override;
-	virtual FString	GetReplayID() const override { return TEXT( "" ); }
+	virtual FString	GetReplayID() const override { return CurrentStreamName; }
 	virtual void SetTimeBufferHintSeconds(const float InTimeBufferHintSeconds) override { TimeBufferHintSeconds = InTimeBufferHintSeconds; }
 	virtual void RefreshHeader() override {};
 	virtual void DownloadHeader(const FDownloadHeaderCallback& Delegate) override
diff --git a/Engine/Source/Runtime/NetworkReplayStreaming/LocalFileNetworkReplayStreaming/Private/LocalFileNetworkReplayStreaming.cpp b/Engine/Source/Runtime/NetworkReplayStreaming/LocalFileNetworkReplayStreaming/Private/LocalFileNetworkReplayStreaming.cpp
--- a/Engine/Source/Runtime/NetworkReplayStreaming/LocalFileNetworkReplayStreaming/Private/LocalFileNetworkReplayStreaming.cpp
+++ b/Engine/Source/Runtime/NetworkReplayStreaming/LocalFileNetworkReplayStreaming/Private/LocalFileNetworkReplayStreaming.cpp
@@ -1,7 +1,6 @@
 // Copyright 1998-2019 Epic Games, Inc. All Rights Reserved.
 
 #include "LocalFileNetworkReplayStreaming.h"
-#include "Misc/NetworkVersion.h"
 #include "HAL/FileManager.h"
 #include "HAL/LowLevelMemTracker.h"
 #include "Misc/Paths.h"
@@ -583,22 +582,10 @@ void FLocalFileNetworkReplayStreamer::FixupFriendlyNameLength(const FString& Unf
 	}
 }
 
-void FLocalFileNetworkReplayStreamer::StartStreaming(const FString& CustomName, const FString& FriendlyName, const TArray<int32>& UserIndices, bool bRecord, const FNetworkReplayVersion& ReplayVersion, const FStartStreamingCallback& Delegate)
-{
-	StartStreaming_Internal(CustomName, FriendlyName, UserIndices, bRecord, ReplayVersion, Delegate);
-}
-
-void FLocalFileNetworkReplayStreamer::StartStreaming(const FString& CustomName, const FString& FriendlyName, const TArray<FString>& UserStrings, bool bRecord, const FNetworkReplayVersion& ReplayVersion, const FStartStreamingCallback& Delegate)
-{
-	TArray<int32> UserIndices;
-	GetUserIndicesFromUserStrings(UserStrings, UserIndices);
-	StartStreaming_Internal(CustomName, FriendlyName, UserIndices, bRecord, ReplayVersion, Delegate);
-}
-
-void FLocalFileNetworkReplayStreamer::StartStreaming_Internal(const FString& CustomName, const FString& FriendlyName, const TArray<int32>& UserIndices, bool bRecord, const FNetworkReplayVersion& ReplayVersion, const FStartStreamingCallback& Delegate)
+void FLocalFileNetworkReplayStreamer::StartStreaming(const FStartStreamingParameters& Params, const FStartStreamingCallback& Delegate)
 {
 	FStartStreamingResult Result;
-	Result.bRecording = bRecord;
+	Result.bRecording = Params.bRecord;
 
 	if (IsStreaming())
 	{
@@ -614,11 +601,11 @@ void FLocalFileNetworkReplayStreamer::StartStreaming_Internal(const FString& Cus
 		return;
 	}
 
-	FString FinalDemoName = CustomName;
+	FString FinalDemoName = Params.CustomName;
 
-	if (CustomName.IsEmpty())
+	if (Params.CustomName.IsEmpty())
 	{
-		if (bRecord)
+		if (Params.bRecord)
 		{
 			// If we're recording and the caller didn't provide a name, generate one automatically
 			FinalDemoName = GetAutomaticDemoName();
@@ -633,7 +620,7 @@ void FLocalFileNetworkReplayStreamer::StartStreaming_Internal(const FString& Cus
 	}
 
 	// Setup the archives
-	StreamAr.SetIsLoading(!bRecord);
+	StreamAr.SetIsLoading(!Params.bRecord);
 	StreamAr.SetIsSaving(!StreamAr.IsLoading());
 	StreamAr.bAtEndOfReplay = false;
 
@@ -656,16 +643,16 @@ void FLocalFileNetworkReplayStreamer::StartStreaming_Internal(const FString& Cus
 
 	CurrentStreamName = FinalDemoName;
 
-	if (!bRecord)
+	if (!Params.bRecord)
 	{
 		// We are playing
 		StreamerState = EStreamerState::Playback;
 
 		// Add the request to start loading
 		AddDelegateFileRequestToQueue<FStartStreamingResult>(EQueuedLocalFileRequestType::StartPlayback, 
-			[this, bRecord, FullDemoFilename](TLocalFileRequestCommonData<FStartStreamingResult>& RequestData)
+			[this, FullDemoFilename](TLocalFileRequestCommonData<FStartStreamingResult>& RequestData)
 			{
-				RequestData.DelegateResult.bRecording = bRecord;
+				RequestData.DelegateResult.bRecording = false;
 				
 				if (!FPaths::FileExists(FullDemoFilename))
 				{
@@ -677,7 +664,7 @@ void FLocalFileNetworkReplayStreamer::StartStreaming_Internal(const FString& Cus
 					ReadReplayInfo(CurrentStreamName, RequestData.ReplayInfo);
 				}
 			},
-			[this, bRecord, Delegate](TLocalFileRequestCommonData<FStartStreamingResult>& RequestData)
+			[this, Delegate](TLocalFileRequestCommonData<FStartStreamingResult>& RequestData)
 			{
 				if (RequestData.DelegateResult.Result == EStreamingOperationResult::ReplayNotFound)
 				{
@@ -698,9 +685,9 @@ void FLocalFileNetworkReplayStreamer::StartStreaming_Internal(const FString& Cus
 						DownloadHeader(FDownloadHeaderCallback());
 
 						AddDelegateFileRequestToQueue<FStartStreamingCallback, FStartStreamingResult>(EQueuedLocalFileRequestType::StartPlayback, Delegate,
-							[this, bRecord](TLocalFileRequestCommonData<FStartStreamingResult>& PlaybackRequestData)
+							[this](TLocalFileRequestCommonData<FStartStreamingResult>& PlaybackRequestData)
 							{
-								PlaybackRequestData.DelegateResult.bRecording = bRecord;	
+								PlaybackRequestData.DelegateResult.bRecording = false;	
 
 								if (CurrentReplayInfo.bIsValid)
 								{
@@ -717,11 +704,11 @@ void FLocalFileNetworkReplayStreamer::StartStreaming_Internal(const FString& Cus
 		StreamerState = EStreamerState::Recording;
 
 		AddDelegateFileRequestToQueue<FStartStreamingResult>(EQueuedLocalFileRequestType::StartRecording,
-			[this, bRecord, FullDemoFilename, ReplayVersion, FriendlyName, FinalDemoName](TLocalFileRequestCommonData<FStartStreamingResult>& RequestData)
+			[this, FullDemoFilename, Params, FinalDemoName](TLocalFileRequestCommonData<FStartStreamingResult>& RequestData)
 			{
 				SCOPE_CYCLE_COUNTER(STAT_LocalReplay_StartRecording);
 
-				RequestData.DelegateResult.bRecording = bRecord;
+				RequestData.DelegateResult.bRecording = true;
 
 				FLocalFileReplayInfo ExistingInfo;
 				if (ReadReplayInfo(FinalDemoName, ExistingInfo))
@@ -735,9 +722,9 @@ void FLocalFileNetworkReplayStreamer::StartStreaming_Internal(const FString& Cus
 				// Delete any existing demo with this name
 				IFileManager::Get().Delete(*FullDemoFilename);
 
-				RequestData.ReplayInfo.NetworkVersion = ReplayVersion.NetworkVersion;
-				RequestData.ReplayInfo.Changelist = ReplayVersion.Changelist;
-				RequestData.ReplayInfo.FriendlyName = FriendlyName;
+				RequestData.ReplayInfo.NetworkVersion = Params.ReplayVersion.NetworkVersion;
+				RequestData.ReplayInfo.Changelist = Params.ReplayVersion.Changelist;
+				RequestData.ReplayInfo.FriendlyName = Params.FriendlyName;
 				RequestData.ReplayInfo.bIsLive = true;
 				RequestData.ReplayInfo.Timestamp = FDateTime::Now();
 
@@ -953,21 +940,6 @@ void FLocalFileNetworkReplayStreamer::DeleteFinishedStream_Internal(const FStrin
 }
 
 void FLocalFileNetworkReplayStreamer::EnumerateStreams(const FNetworkReplayVersion& ReplayVersion, const int32 UserIndex, const FString& MetaString, const TArray< FString >& ExtraParms, const FEnumerateStreamsCallback& Delegate)
-{
-	EnumerateStreams_Internal(ReplayVersion, UserIndex, MetaString, ExtraParms, Delegate);
-}
-
-void FLocalFileNetworkReplayStreamer::EnumerateStreams(const FNetworkReplayVersion& ReplayVersion, const FString& UserString, const FString& MetaString, const TArray< FString >& ExtraParms, const FEnumerateStreamsCallback& Delegate)
-{
-	EnumerateStreams_Internal(ReplayVersion, GetUserIndexFromUserString(UserString), MetaString, ExtraParms, Delegate);
-}
-
-void FLocalFileNetworkReplayStreamer::EnumerateStreams(const FNetworkReplayVersion& ReplayVersion, const FString& UserString, const FString& MetaString, const FEnumerateStreamsCallback& Delegate)
-{
-	EnumerateStreams_Internal(ReplayVersion, GetUserIndexFromUserString(UserString), MetaString, TArray< FString >(), Delegate);
-}
-
-void FLocalFileNetworkReplayStreamer::EnumerateStreams_Internal(const FNetworkReplayVersion& ReplayVersion, const int32 UserIndex, const FString& MetaString, const TArray< FString >& ExtraParms, const FEnumerateStreamsCallback& Delegate)
 {
 	AddDelegateFileRequestToQueue<FEnumerateStreamsCallback, FEnumerateStreamsResult>(EQueuedLocalFileRequestType::EnumeratingStreams, Delegate,
 		[this, ReplayVersion](TLocalFileRequestCommonData<FEnumerateStreamsResult>& RequestData)
@@ -1022,16 +994,6 @@ void FLocalFileNetworkReplayStreamer::EnumerateStreams_Internal(const FNetworkRe
 }
 
 void FLocalFileNetworkReplayStreamer::EnumerateRecentStreams(const FNetworkReplayVersion& ReplayVersion, const int32 UserIndex, const FEnumerateStreamsCallback& Delegate)
-{
-	EnumerateRecentStreams_Internal(ReplayVersion, UserIndex, Delegate);
-}
-
-void FLocalFileNetworkReplayStreamer::EnumerateRecentStreams(const FNetworkReplayVersion& ReplayVersion, const FString& RecentViewer, const FEnumerateStreamsCallback& Delegate)
-{
-	EnumerateRecentStreams_Internal(ReplayVersion, GetUserIndexFromUserString(RecentViewer), Delegate);
-}
-
-void FLocalFileNetworkReplayStreamer::EnumerateRecentStreams_Internal(const FNetworkReplayVersion& ReplayVersion, const int32 UserIndex, const FEnumerateStreamsCallback& Delegate)
 {
 	UE_LOG(LogLocalFileReplay, Log, TEXT("FLocalFileNetworkReplayStreamer::EnumerateRecentStreams is currently unsupported."));
 	FEnumerateStreamsResult Result;
@@ -2987,110 +2949,6 @@ void FGenericQueuedLocalFileRequest::FinishRequest()
 	}
 }
 
-const FString FLocalFileNetworkReplayStreamer::GetUserStringFromUserIndex(const int32 UserIndex)
-{
-	if (UserIndex != INDEX_NONE && GEngine != nullptr)
-	{
-		if (UWorld* World = GWorld.GetReference())
-		{
-			if (ULocalPlayer const * const LocalPlayer = GEngine->GetLocalPlayerFromControllerId(World, UserIndex))
-			{
-				return LocalPlayer->GetPreferredUniqueNetId().ToString();
-			}
-		}
-	}
-
-	return FString();
-}
-
-const void FLocalFileNetworkReplayStreamer::GetUserStringsFromUserIndices(const TArray<int32>& UserIndices, TArray<FString>& OutUserStrings)
-{
-	if (GEngine != nullptr)
-	{
-		if (UserIndices.Num() == 1)
-		{
-			OutUserStrings.Emplace(GetUserStringFromUserIndex(UserIndices[0]));
-		}
-		else if (UserIndices.Num() > 1)
-		{
-			if (UWorld* World = GWorld.GetReference())
-			{
-				TMap<int32, FString> IdToString;
-				for (auto ConstIt = GEngine->GetLocalPlayerIterator(World); ConstIt; ++ConstIt)
-				{
-					if (ULocalPlayer const * const LocalPlayer = *ConstIt)
-					{
-						IdToString.Emplace(LocalPlayer->GetControllerId(), LocalPlayer->GetPreferredUniqueNetId().ToString());
-					}
-				}
-
-				for (const int32 UserIndex : UserIndices)
-				{
-					if (FString const * const UserString = IdToString.Find(UserIndex))
-					{
-						OutUserStrings.Emplace(*UserString);
-					}
-				}
-			}
-		}
-	}
-}
-
-const int32 FLocalFileNetworkReplayStreamer::GetUserIndexFromUserString(const FString& UserString)
-{
-	if (!UserString.IsEmpty() && GEngine != nullptr)
-	{
-		if (UWorld* World = GWorld.GetReference())
-		{
-			for (auto ConstIt = GEngine->GetLocalPlayerIterator(World); ConstIt; ++ConstIt)
-			{
-				if (ULocalPlayer const * const LocalPlayer = *ConstIt)
-				{
-					if (UserString.Equals(LocalPlayer->GetPreferredUniqueNetId().ToString()))
-					{
-						return LocalPlayer->GetControllerId();
-					}
-				}
-			}
-		}
-	}
-
-	return INDEX_NONE;
-}
-
-const void FLocalFileNetworkReplayStreamer::GetUserIndicesFromUserStrings(const TArray<FString>& UserStrings, TArray<int32>& OutUserIndices)
-{
-	if (GEngine != nullptr)
-	{
-		if (UserStrings.Num() == 1)
-		{
-			OutUserIndices.Emplace(GetUserIndexFromUserString(UserStrings[0]));
-		}
-		else if (UserStrings.Num() > 1)
-		{
-			if (UWorld* World = GWorld.GetReference())
-			{
-				TMap<FString, int32> StringToId;
-				for (auto ConstIt = GEngine->GetLocalPlayerIterator(World); ConstIt; ++ConstIt)
-				{
-					if (ULocalPlayer const * const LocalPlayer = *ConstIt)
-					{
-						StringToId.Emplace(LocalPlayer->GetPreferredUniqueNetId().ToString(), LocalPlayer->GetControllerId());
-					}
-				}
-
-				for (const FString& UserString : UserStrings)
-				{
-					if (int32 const * const UserIndex = StringToId.Find(UserString))
-					{
-						OutUserIndices.Emplace(*UserIndex);
-					}
-				}
-			}
-		}
-	}
-}
-
 bool FLocalFileNetworkReplayStreamer::IsCheckpointTypeSupported(EReplayCheckpointType CheckpointType) const
 {
 	bool bSupported = false;
diff --git a/Engine/Source/Runtime/NetworkReplayStreaming/LocalFileNetworkReplayStreaming/Public/LocalFileNetworkReplayStreaming.h b/Engine/Source/Runtime/NetworkReplayStreaming/LocalFileNetworkReplayStreaming/Public/LocalFileNetworkReplayStreaming.h
--- a/Engine/Source/Runtime/NetworkReplayStreaming/LocalFileNetworkReplayStreaming/Public/LocalFileNetworkReplayStreaming.h
+++ b/Engine/Source/Runtime/NetworkReplayStreaming/LocalFileNetworkReplayStreaming/Public/LocalFileNetworkReplayStreaming.h
@@ -398,8 +398,7 @@ public:
 	virtual ~FLocalFileNetworkReplayStreamer();
 
 	/** INetworkReplayStreamer implementation */
-	virtual void StartStreaming(const FString& CustomName, const FString& FriendlyName, const TArray< FString >& UserNames, bool bRecord, const FNetworkReplayVersion& ReplayVersion, const FStartStreamingCallback& Delegate) override;
-	virtual void StartStreaming(const FString& CustomName, const FString& FriendlyName, const TArray< int32 >& UserIndices, bool bRecord, const FNetworkReplayVersion& ReplayVersion, const FStartStreamingCallback& Delegate) override;
+	virtual void StartStreaming(const FStartStreamingParameters& Params, const FStartStreamingCallback& Delegate) override;
 	virtual void StopStreaming() override;
 	virtual FArchive* GetHeaderArchive() override;
 	virtual FArchive* GetStreamingArchive() override;
@@ -417,10 +416,7 @@ public:
 	virtual bool IsLive() const override;
 	virtual void DeleteFinishedStream(const FString& StreamName, const FDeleteFinishedStreamCallback& Delegate) override;
 	virtual void DeleteFinishedStream( const FString& StreamName, const int32 UserIndex, const FDeleteFinishedStreamCallback& Delegate ) override;
-	virtual void EnumerateStreams(const FNetworkReplayVersion& ReplayVersion, const FString& UserString, const FString& MetaString, const FEnumerateStreamsCallback& Delegate) override;
-	virtual void EnumerateStreams(const FNetworkReplayVersion& InReplayVersion, const FString& UserString, const FString& MetaString, const TArray< FString >& ExtraParms, const FEnumerateStreamsCallback& Delegate) override;
 	virtual void EnumerateStreams( const FNetworkReplayVersion& InReplayVersion, const int32 UserIndex, const FString& MetaString, const TArray< FString >& ExtraParms, const FEnumerateStreamsCallback& Delegate ) override;
-	virtual void EnumerateRecentStreams(const FNetworkReplayVersion& ReplayVersion, const FString& RecentViewer, const FEnumerateStreamsCallback& Delegate) override;
 	virtual void EnumerateRecentStreams( const FNetworkReplayVersion& ReplayVersion, const int32 UserIndex, const FEnumerateStreamsCallback& Delegate ) override;
 	virtual ENetworkReplayError::Type GetLastError() const override;
 	virtual void AddUserToReplay(const FString& UserString) override;
@@ -442,7 +438,7 @@ public:
 	virtual void RenameReplayFriendlyName(const FString& ReplayName, const FString& NewFriendlyName, const int32 UserIndex, const FRenameReplayCallback& Delegate) override;
 	virtual void RenameReplay(const FString& ReplayName, const FString& NewName, const FRenameReplayCallback& Delegate) override;
 	virtual void RenameReplay(const FString& ReplayName, const FString& NewName, const int32 UserIndex, const FRenameReplayCallback& Delegate) override;
-	virtual FString	GetReplayID() const override { return TEXT(""); }
+	virtual FString	GetReplayID() const override { return CurrentStreamName; }
 	virtual void SetTimeBufferHintSeconds(const float InTimeBufferHintSeconds) override {}
 	virtual void RefreshHeader() override;
 	virtual void DownloadHeader(const FDownloadHeaderCallback& Delegate) override;
@@ -524,16 +520,7 @@ public:
 
 protected:
 
-	static const FString GetUserStringFromUserIndex(const int32 UserIndex);
-	static const void GetUserStringsFromUserIndices(const TArray<int32>& UserIndices, TArray<FString>& OutUserStrings);
-
-	static const int32 GetUserIndexFromUserString(const FString& UserString);
-	static const void GetUserIndicesFromUserStrings(const TArray<FString>& UserStrings, TArray<int32>& OutUserIndices);
-
-	void StartStreaming_Internal(const FString& CustomName, const FString& FriendlyName, const TArray<int32>& UserIndices, bool bRecord, const FNetworkReplayVersion& ReplayVersion, const FStartStreamingCallback& Delegate);
 	void DeleteFinishedStream_Internal(const FString& StreamName, const int32 UserIndex, const FDeleteFinishedStreamCallback& Delegate);
-	void EnumerateStreams_Internal(const FNetworkReplayVersion& InReplayVersion, const int32 UserIndex, const FString& MetaString, const TArray< FString >& ExtraParms, const FEnumerateStreamsCallback& Delegate);
-	void EnumerateRecentStreams_Internal(const FNetworkReplayVersion& ReplayVersion, const int32 UserIndex, const FEnumerateStreamsCallback& Delegate);
 	void EnumerateEvents_Internal(const FString& ReplayName, const FString& Group, const int32 UserIndex, const FEnumerateEventsCallback& Delegate);
 	void RequestEventData_Internal(const FString& ReplayName, const FString& EventId, const int32 UserIndex, const FRequestEventDataCallback& Delegate);
 	void KeepReplay_Internal(const FString& ReplayName, const bool bKeep, const int32 UserIndex, const FKeepReplayCallback& Delegate);
diff --git a/Engine/Source/Runtime/NetworkReplayStreaming/NetworkReplayStreaming/NetworkReplayStreaming.Build.cs b/Engine/Source/Runtime/NetworkReplayStreaming/NetworkReplayStreaming/NetworkReplayStreaming.Build.cs
--- a/Engine/Source/Runtime/NetworkReplayStreaming/NetworkReplayStreaming/NetworkReplayStreaming.Build.cs
+++ b/Engine/Source/Runtime/NetworkReplayStreaming/NetworkReplayStreaming/NetworkReplayStreaming.Build.cs
@@ -12,6 +12,7 @@ namespace UnrealBuildTool.Rules
 				new string[]
 				{
 					"Core",
+					"Engine",
                     "Json",
                     // ... add other public dependencies that you statically link with here ...
 				}
diff --git a/Engine/Source/Runtime/NetworkReplayStreaming/NetworkReplayStreaming/Private/NetworkReplayStreaming.cpp b/Engine/Source/Runtime/NetworkReplayStreaming/NetworkReplayStreaming/Private/NetworkReplayStreaming.cpp
--- a/Engine/Source/Runtime/NetworkReplayStreaming/NetworkReplayStreaming/Private/NetworkReplayStreaming.cpp
+++ b/Engine/Source/Runtime/NetworkReplayStreaming/NetworkReplayStreaming/Private/NetworkReplayStreaming.cpp
@@ -123,3 +123,97 @@ bool FNetworkReplayStreaming::Exec(UWorld* InWorld, const TCHAR* Cmd, FOutputDev
 
 	return false;
 }
+
+void INetworkReplayStreamer::StartStreaming(const FString& CustomName, const FString& FriendlyName, const TArray<FString>& UserNames, bool bRecord, const FNetworkReplayVersion& ReplayVersion, const FStartStreamingCallback& Delegate)
+{
+	FStartStreamingParameters Params;
+	Params.CustomName = CustomName;
+	Params.FriendlyName = FriendlyName;
+	GetUserIndicesFromUserStrings(UserNames, Params.UserIndices);
+	Params.bRecord = bRecord;
+	Params.ReplayVersion = ReplayVersion;
+
+	StartStreaming(Params, Delegate);
+}
+
+void INetworkReplayStreamer::StartStreaming(const FString& CustomName, const FString& FriendlyName, const TArray<int32>& UserIndices, bool bRecord, const FNetworkReplayVersion& ReplayVersion, const FStartStreamingCallback& Delegate)
+{
+	FStartStreamingParameters Params;
+	Params.CustomName = CustomName;
+	Params.FriendlyName = FriendlyName;
+	Params.UserIndices = UserIndices;
+	Params.bRecord = bRecord;
+	Params.ReplayVersion = ReplayVersion;
+
+	StartStreaming(Params, Delegate);
+}
+
+void INetworkReplayStreamer::EnumerateStreams(const FNetworkReplayVersion& ReplayVersion, const FString& UserString, const FString& MetaString, const FEnumerateStreamsCallback& Delegate)
+{
+	EnumerateStreams(ReplayVersion, GetUserIndexFromUserString(UserString), MetaString, TArray<FString>(), Delegate);
+}
+
+void INetworkReplayStreamer::EnumerateStreams(const FNetworkReplayVersion& ReplayVersion, const FString& UserString, const FString& MetaString, const TArray<FString>& ExtraParms, const FEnumerateStreamsCallback& Delegate)
+{
+	EnumerateStreams(ReplayVersion, GetUserIndexFromUserString(UserString), MetaString, ExtraParms, Delegate);
+}
+
+void INetworkReplayStreamer::EnumerateRecentStreams(const FNetworkReplayVersion& ReplayVersion, const FString& RecentViewer, const FEnumerateStreamsCallback& Delegate)
+{
+	EnumerateRecentStreams(ReplayVersion, GetUserIndexFromUserString(RecentViewer), Delegate);
+}
+
+const int32 INetworkReplayStreamer::GetUserIndexFromUserString(const FString& UserString)
+{
+	if (!UserString.IsEmpty() && GEngine != nullptr)
+	{
+		if (UWorld* World = GWorld.GetReference())
+		{
+			for (auto ConstIt = GEngine->GetLocalPlayerIterator(World); ConstIt; ++ConstIt)
+			{
+				if (ULocalPlayer const * const LocalPlayer = *ConstIt)
+				{
+					if (UserString.Equals(LocalPlayer->GetPreferredUniqueNetId().ToString()))
+					{
+						return LocalPlayer->GetControllerId();
+					}
+				}
+			}
+		}
+	}
+
+	return INDEX_NONE;
+}
+
+const void INetworkReplayStreamer::GetUserIndicesFromUserStrings(const TArray<FString>& UserStrings, TArray<int32>& OutUserIndices)
+{
+	if (GEngine != nullptr)
+	{
+		if (UserStrings.Num() == 1)
+		{
+			OutUserIndices.Emplace(GetUserIndexFromUserString(UserStrings[0]));
+		}
+		else if (UserStrings.Num() > 1)
+		{
+			if (UWorld* World = GWorld.GetReference())
+			{
+				TMap<FString, int32> StringToId;
+				for (auto ConstIt = GEngine->GetLocalPlayerIterator(World); ConstIt; ++ConstIt)
+				{
+					if (ULocalPlayer const * const LocalPlayer = *ConstIt)
+					{
+						StringToId.Emplace(LocalPlayer->GetPreferredUniqueNetId().ToString(), LocalPlayer->GetControllerId());
+					}
+				}
+
+				for (const FString& UserString : UserStrings)
+				{
+					if (int32 const * const UserIndex = StringToId.Find(UserString))
+					{
+						OutUserIndices.Emplace(*UserIndex);
+					}
+				}
+			}
+		}
+	}
+}
\ No newline at end of file
diff --git a/Engine/Source/Runtime/NetworkReplayStreaming/NetworkReplayStreaming/Public/NetworkReplayStreaming.h b/Engine/Source/Runtime/NetworkReplayStreaming/NetworkReplayStreaming/Public/NetworkReplayStreaming.h
--- a/Engine/Source/Runtime/NetworkReplayStreaming/NetworkReplayStreaming/Public/NetworkReplayStreaming.h
+++ b/Engine/Source/Runtime/NetworkReplayStreaming/NetworkReplayStreaming/Public/NetworkReplayStreaming.h
@@ -5,6 +5,8 @@
 // Dependencies.
 
 #include "CoreMinimal.h"
+#include "Engine/EngineBaseTypes.h"
+#include "Misc/NetworkVersion.h"
 #include "Modules/ModuleInterface.h"
 #include "Modules/ModuleManager.h"
 #include "Serialization/JsonSerializerMacros.h"
@@ -463,6 +465,16 @@ enum class EReplayCheckpointType : uint8
 	Delta
 };
 
+struct FStartStreamingParameters
+{
+	FString CustomName;
+	FString FriendlyName;
+	FURL DemoURL;
+	TArray<int32> UserIndices;
+	FNetworkReplayVersion ReplayVersion;
+	bool bRecord;
+};
+
 /**
  * Generic interface for network replay streaming
  *
@@ -470,13 +482,18 @@ enum class EReplayCheckpointType : uint8
  * that delegate upon completion, and indicates success / failure through an appropriate
  * result type passed into the delegate.
  */
-class INetworkReplayStreamer 
+class NETWORKREPLAYSTREAMING_API INetworkReplayStreamer 
 {
 public:
 	virtual ~INetworkReplayStreamer() {}
 
-	virtual void StartStreaming(const FString& CustomName, const FString& FriendlyName, const TArray< FString >& UserNames, bool bRecord, const FNetworkReplayVersion& ReplayVersion, const FStartStreamingCallback& Delegate) = 0;
-	virtual void StartStreaming(const FString& CustomName, const FString& FriendlyName, const TArray< int32 >& UserIndices, bool bRecord, const FNetworkReplayVersion& ReplayVersion, const FStartStreamingCallback& Delegate) = 0;
+	UE_DEPRECATED(4.23, "StartStreaming now takes a FStartStreamingParameters argument")
+	virtual void StartStreaming(const FString& CustomName, const FString& FriendlyName, const TArray<FString>& UserNames, bool bRecord, const FNetworkReplayVersion& ReplayVersion, const FStartStreamingCallback& Delegate);
+
+	UE_DEPRECATED(4.23, "StartStreaming now takes a FStartStreamingParameters argument")
+	virtual void StartStreaming(const FString& CustomName, const FString& FriendlyName, const TArray<int32>& UserIndices, bool bRecord, const FNetworkReplayVersion& ReplayVersion, const FStartStreamingCallback& Delegate);
+
+	virtual void StartStreaming(const FStartStreamingParameters& Params, const FStartStreamingCallback& Delegate) = 0;
 
 	virtual void StopStreaming() = 0;
 	virtual FArchive* GetHeaderArchive() = 0;
@@ -512,7 +529,7 @@ public:
 
 	virtual void RequestEventData(const FString& EventID, const FRequestEventDataCallback& Delegate) = 0;
 	virtual void RequestEventData(const FString& ReplayName, const FString& EventID, const FRequestEventDataCallback& Delegate) = 0;
-	virtual void RequestEventData(const FString& ReplayName, const FString& EventId, const int32 UserIndex, const FRequestEventDataCallback& Delegate) = 0;
+	virtual void RequestEventData(const FString& ReplayName, const FString& EventID, const int32 UserIndex, const FRequestEventDataCallback& Delegate) = 0;
 
 	virtual void RequestEventGroupData(const FString& Group, const FRequestEventGroupDataCallback& Delegate) = 0;
 	virtual void RequestEventGroupData(const FString& ReplayName, const FString& Group, const FRequestEventGroupDataCallback& Delegate) = 0;
@@ -554,6 +571,8 @@ public:
 
 	/** Returns true if the playing stream is currently in progress */
 	virtual bool IsLive() const = 0;
+
+	/** Returns the active replay name */
 	virtual FString	GetReplayID() const = 0;
 
 	/**
@@ -563,29 +582,32 @@ public:
 	 * @param Delegate A delegate that will be executed if bound when the delete operation completes
 	 */
 	virtual void DeleteFinishedStream(const FString& StreamName, const FDeleteFinishedStreamCallback& Delegate) = 0;
-	virtual void DeleteFinishedStream( const FString& StreamName, const int32 UserIndex, const FDeleteFinishedStreamCallback& Delegate ) = 0;
+	virtual void DeleteFinishedStream(const FString& StreamName, const int32 UserIndex, const FDeleteFinishedStreamCallback& Delegate) = 0;
 
 	/**
 	 * Retrieves the streams that are available for viewing. May execute asynchronously.
 	 *
 	 * @param Delegate A delegate that will be executed if bound when the list of streams is available
 	 */
-	virtual void EnumerateStreams(const FNetworkReplayVersion& ReplayVersion, const FString& UserString, const FString& MetaString, const FEnumerateStreamsCallback& Delegate) = 0;
+	UE_DEPRECATED(4.23, "UserString is deprecated in favor of UserIndex")
+	virtual void EnumerateStreams(const FNetworkReplayVersion& ReplayVersion, const FString& UserString, const FString& MetaString, const FEnumerateStreamsCallback& Delegate);
 
 	/**
 	* Retrieves the streams that are available for viewing. May execute asynchronously.
 	* Allows the caller to pass in a custom list of query parameters
 	*/
-	virtual void EnumerateStreams(const FNetworkReplayVersion& ReplayVersion, const FString& UserString, const FString& MetaString, const TArray<FString>& ExtraParms, const FEnumerateStreamsCallback& Delegate) = 0;
-	virtual void EnumerateStreams( const FNetworkReplayVersion& InReplayVersion, const int32 UserIndex, const FString& MetaString, const TArray< FString >& ExtraParms, const FEnumerateStreamsCallback& Delegate ) = 0;
+	UE_DEPRECATED(4.23, "UserString is deprecated in favor of UserIndex")
+	virtual void EnumerateStreams(const FNetworkReplayVersion& ReplayVersion, const FString& UserString, const FString& MetaString, const TArray<FString>& ExtraParms, const FEnumerateStreamsCallback& Delegate);
+	virtual void EnumerateStreams(const FNetworkReplayVersion& ReplayVersion, const int32 UserIndex, const FString& MetaString, const TArray< FString >& ExtraParms, const FEnumerateStreamsCallback& Delegate) = 0;
 
 	/**
 	 * Retrieves the streams that have been recently viewed. May execute asynchronously.
 	 *
 	 * @param Delegate A delegate that will be executed if bound when the list of streams is available
 	 */
-	virtual void EnumerateRecentStreams(const FNetworkReplayVersion& ReplayVersion, const FString& RecentViewer, const FEnumerateStreamsCallback& Delegate) = 0;
-	virtual void EnumerateRecentStreams( const FNetworkReplayVersion& ReplayVersion, const int32 UserIndex, const FEnumerateStreamsCallback& Delegate ) = 0;
+	UE_DEPRECATED(4.23, "RecentViewer is deprecated in favor of UserIndex")
+	virtual void EnumerateRecentStreams(const FNetworkReplayVersion& ReplayVersion, const FString& RecentViewer, const FEnumerateStreamsCallback& Delegate);
+	virtual void EnumerateRecentStreams(const FNetworkReplayVersion& ReplayVersion, const int32 UserIndex, const FEnumerateStreamsCallback& Delegate) = 0;
 
 	/** Returns the last error that occurred while streaming replays */
 	virtual ENetworkReplayError::Type GetLastError() const = 0;
@@ -624,6 +646,10 @@ public:
 	virtual void SetAnalyticsProvider(TSharedPtr<IAnalyticsProvider>& InProvider) {}
 
 	virtual void Exec(const TCHAR* Cmd, FOutputDevice& Ar) {}
+
+private:
+	static const int32 GetUserIndexFromUserString(const FString& UserString);
+	static const void GetUserIndicesFromUserStrings(const TArray<FString>& UserStrings, TArray<int32>& OutUserIndices);
 };
 
 /** Replay streamer factory */
diff --git a/Engine/Source/Runtime/NetworkReplayStreaming/NullNetworkReplayStreaming/Private/NullNetworkReplayStreaming.cpp b/Engine/Source/Runtime/NetworkReplayStreaming/NullNetworkReplayStreaming/Private/NullNetworkReplayStreaming.cpp
--- a/Engine/Source/Runtime/NetworkReplayStreaming/NullNetworkReplayStreaming/Private/NullNetworkReplayStreaming.cpp
+++ b/Engine/Source/Runtime/NetworkReplayStreaming/NullNetworkReplayStreaming/Private/NullNetworkReplayStreaming.cpp
@@ -3,7 +3,6 @@
 #include "NullNetworkReplayStreaming.h"
 #include "HAL/FileManager.h"
 #include "Misc/Paths.h"
-#include "Misc/NetworkVersion.h"
 #include "Misc/EngineVersion.h"
 
 DEFINE_LOG_CATEGORY_STATIC( LogNullReplay, Log, All );
@@ -176,21 +175,16 @@ FString GetAutomaticDemoName()
 	return FinalDemoName;
 }
 
-void FNullNetworkReplayStreamer::StartStreaming(const FString& CustomName, const FString& FriendlyName, const TArray< int32 >& UserIndices, bool bRecord, const FNetworkReplayVersion& ReplayVersion, const FStartStreamingCallback& Delegate)
+void FNullNetworkReplayStreamer::StartStreaming(const FStartStreamingParameters& Params, const FStartStreamingCallback& Delegate)
 {
-	StartStreaming(CustomName, FriendlyName, TArray<FString>(), bRecord, ReplayVersion, Delegate);
-}
-
-void FNullNetworkReplayStreamer::StartStreaming( const FString& CustomName, const FString& FriendlyName, const TArray< FString >& UserNames, bool bRecord, const FNetworkReplayVersion& ReplayVersion, const FStartStreamingCallback& Delegate )
-{
-	FString FinalDemoName = CustomName;
+	FString FinalDemoName = Params.CustomName;
 
 	FStartStreamingResult Result;
-	Result.bRecording = bRecord;
+	Result.bRecording = Params.bRecord;
 
-	if ( CustomName.IsEmpty() )
+	if ( Params.CustomName.IsEmpty() )
 	{
-		if ( bRecord )
+		if ( Params.bRecord )
 		{
 			// If we're recording and the caller didn't provide a name, generate one automatically
 			FinalDemoName = GetAutomaticDemoName();
@@ -211,7 +205,7 @@ void FNullNetworkReplayStreamer::StartStreaming( const FString& CustomName, cons
 	
 	CurrentStreamName = FinalDemoName;
 
-	if ( !bRecord )
+	if ( !Params.bRecord )
 	{
 		// Load metadata if it exists
 		ReplayInfo = ReadReplayInfo( CurrentStreamName );
@@ -237,9 +231,9 @@ void FNullNetworkReplayStreamer::StartStreaming( const FString& CustomName, cons
 		CurrentCheckpointIndex = 0;
 
 		// Set up replay info
-		ReplayInfo.NetworkVersion = ReplayVersion.NetworkVersion;
-		ReplayInfo.Changelist = ReplayVersion.Changelist;
-		ReplayInfo.FriendlyName = FriendlyName;
+		ReplayInfo.NetworkVersion = Params.ReplayVersion.NetworkVersion;
+		ReplayInfo.Changelist = Params.ReplayVersion.Changelist;
+		ReplayInfo.FriendlyName = Params.FriendlyName;
 
 		WriteReplayInfo(CurrentStreamName, ReplayInfo);
 	}
@@ -335,16 +329,6 @@ void FNullNetworkReplayStreamer::DeleteFinishedStream( const FString& StreamName
 }
 
 void FNullNetworkReplayStreamer::EnumerateStreams(const FNetworkReplayVersion& ReplayVersion, const int32 UserIndex, const FString& MetaString, const TArray< FString >& ExtraParms, const FEnumerateStreamsCallback& Delegate)
-{
-	EnumerateStreams(ReplayVersion, FString(), MetaString, ExtraParms, Delegate);
-}
-
-void FNullNetworkReplayStreamer::EnumerateStreams( const FNetworkReplayVersion& ReplayVersion, const FString& UserString, const FString& MetaString, const FEnumerateStreamsCallback& Delegate )
-{
-	EnumerateStreams( ReplayVersion, UserString, MetaString, TArray< FString >(), Delegate );
-}
-
-void FNullNetworkReplayStreamer::EnumerateStreams( const FNetworkReplayVersion& ReplayVersion, const FString& UserString, const FString& MetaString, const TArray< FString >& ExtraParms, const FEnumerateStreamsCallback& Delegate )
 {
 	// Simply returns a stream for each folder in the Saved/Demos directory
 	const FString WildCardPath = ::GetDemoPath() + TEXT( "*" );
@@ -516,11 +500,6 @@ void FNullNetworkReplayStreamer::RenameReplay(const FString& ReplayName, const F
 }
 
 void FNullNetworkReplayStreamer::EnumerateRecentStreams(const FNetworkReplayVersion& ReplayVersion, const int32 UserIndex, const FEnumerateStreamsCallback& Delegate)
-{
-	EnumerateRecentStreams(ReplayVersion, FString(), Delegate);
-}
-
-void FNullNetworkReplayStreamer::EnumerateRecentStreams(const FNetworkReplayVersion& ReplayVersion, const FString& RecentViewer, const FEnumerateStreamsCallback& Delegate)
 {
 	UE_LOG(LogNullReplay, Log, TEXT("FNullNetworkReplayStreamer::EnumerateRecentStreams is currently unsupported."));
 
diff --git a/Engine/Source/Runtime/NetworkReplayStreaming/NullNetworkReplayStreaming/Public/NullNetworkReplayStreaming.h b/Engine/Source/Runtime/NetworkReplayStreaming/NullNetworkReplayStreaming/Public/NullNetworkReplayStreaming.h
--- a/Engine/Source/Runtime/NetworkReplayStreaming/NullNetworkReplayStreaming/Public/NullNetworkReplayStreaming.h
+++ b/Engine/Source/Runtime/NetworkReplayStreaming/NullNetworkReplayStreaming/Public/NullNetworkReplayStreaming.h
@@ -44,8 +44,7 @@ public:
 	{}
 	
 	/** INetworkReplayStreamer implementation */
-	virtual void StartStreaming( const FString& CustomName, const FString& FriendlyName, const TArray< FString >& UserNames, bool bRecord, const FNetworkReplayVersion& ReplayVersion, const FStartStreamingCallback& Delegate ) override;
-	virtual void StartStreaming( const FString& CustomName, const FString& FriendlyName, const TArray< int32 >& UserIndices, bool bRecord, const FNetworkReplayVersion& ReplayVersion, const FStartStreamingCallback& Delegate ) override;
+	virtual void StartStreaming(const FStartStreamingParameters& Params, const FStartStreamingCallback& Delegate) override;
 	virtual void StopStreaming() override;
 	virtual FArchive* GetHeaderArchive() override;
 	virtual FArchive* GetStreamingArchive() override;
@@ -63,10 +62,7 @@ public:
 	virtual bool IsLive() const override;
 	virtual void DeleteFinishedStream( const FString& StreamName, const FDeleteFinishedStreamCallback& Delegate) override;
 	virtual void DeleteFinishedStream( const FString& StreamName, const int32 UserIndex, const FDeleteFinishedStreamCallback& Delegate ) override;
-	virtual void EnumerateStreams( const FNetworkReplayVersion& ReplayVersion, const FString& UserString, const FString& MetaString, const FEnumerateStreamsCallback& Delegate ) override;
-	virtual void EnumerateStreams( const FNetworkReplayVersion& InReplayVersion, const FString& UserString, const FString& MetaString, const TArray< FString >& ExtraParms, const FEnumerateStreamsCallback& Delegate ) override;
 	virtual void EnumerateStreams( const FNetworkReplayVersion& InReplayVersion, const int32 UserIndex, const FString& MetaString, const TArray< FString >& ExtraParms, const FEnumerateStreamsCallback& Delegate ) override;
-	virtual void EnumerateRecentStreams( const FNetworkReplayVersion& ReplayVersion, const FString& RecentViewer, const FEnumerateStreamsCallback& Delegate ) override;
 	virtual void EnumerateRecentStreams( const FNetworkReplayVersion& ReplayVersion, const int32 UserIndex, const FEnumerateStreamsCallback& Delegate ) override;
 	virtual ENetworkReplayError::Type GetLastError() const override { return ENetworkReplayError::None; }
 	virtual void AddUserToReplay(const FString& UserString) override;
@@ -90,7 +86,7 @@ public:
 	virtual void RenameReplayFriendlyName(const FString& ReplayName, const FString& NewFriendlyName, const int32 UserIndex, const FRenameReplayCallback& Delegate) override;
 	virtual void RenameReplay(const FString& ReplayName, const FString& NewName, const FRenameReplayCallback& Delegate) override;
 	virtual void RenameReplay(const FString& ReplayName, const FString& NewName, const int32 UserIndex, const FRenameReplayCallback& Delegate) override;
-	virtual FString	GetReplayID() const override { return TEXT( "" ); }
+	virtual FString	GetReplayID() const override { return CurrentStreamName; }
 	virtual void SetTimeBufferHintSeconds(const float InTimeBufferHintSeconds) override {}
 	virtual void RefreshHeader() override {};
 	virtual void DownloadHeader(const FDownloadHeaderCallback& Delegate) override
diff --git a/Engine/Source/Runtime/NetworkReplayStreaming/SaveGameNetworkReplayStreaming/Private/SaveGameNetworkReplayStreaming.cpp b/Engine/Source/Runtime/NetworkReplayStreaming/SaveGameNetworkReplayStreaming/Private/SaveGameNetworkReplayStreaming.cpp
--- a/Engine/Source/Runtime/NetworkReplayStreaming/SaveGameNetworkReplayStreaming/Private/SaveGameNetworkReplayStreaming.cpp
+++ b/Engine/Source/Runtime/NetworkReplayStreaming/SaveGameNetworkReplayStreaming/Private/SaveGameNetworkReplayStreaming.cpp
@@ -10,7 +10,6 @@
 #include "HAL/IConsoleManager.h"
 #include "Async/Async.h"
 #include "Tickable.h"
-#include "Misc/NetworkVersion.h"
 #include "GameDelegates.h"
 #include "Engine/World.h"
 #include "Engine/Engine.h"
@@ -635,35 +634,16 @@ FSaveGameNetworkReplayStreamer::FSaveGameNetworkReplayStreamer(const FString& In
 {
 }
 
-void FSaveGameNetworkReplayStreamer::StartStreaming(const FString& CustomName, const FString& FriendlyName, const TArray<int32>& UserIndices, bool bRecord, const FNetworkReplayVersion& ReplayVersion, const FStartStreamingCallback& Delegate)
-{
-	if (!IsSaveGameFileName(CustomName))
-	{
-		FLocalFileNetworkReplayStreamer::StartStreaming(CustomName, FriendlyName, UserIndices, bRecord, ReplayVersion, Delegate);
-	}
-	else if (UserIndices.Num() > 0 && UserIndices[0] != INDEX_NONE)
-	{
-		StartStreamingSaved(CustomName, FriendlyName, UserIndices, bRecord, ReplayVersion, Delegate);
-	}
-	else
-	{
-		UE_LOG(LogSaveGameReplay, Warning, TEXT("FSaveGameNetworkReplayStreamer::StartStreaming - Invalid UserIndex"));
-		Delegate.ExecuteIfBound(FStartStreamingResult());
-	}
-}
-
-void FSaveGameNetworkReplayStreamer::StartStreaming(const FString& CustomName, const FString& FriendlyName, const TArray<FString>& UserStrings, bool bRecord, const FNetworkReplayVersion& ReplayVersion, const FStartStreamingCallback& Delegate)
+void FSaveGameNetworkReplayStreamer::StartStreaming(const FStartStreamingParameters& Params, const FStartStreamingCallback& Delegate)
 {
 	// If we're not handling a SaveFile directly, then just do normal streaming behavior.
-	if (!IsSaveGameFileName(CustomName))
+	if (!IsSaveGameFileName(Params.CustomName))
 	{
-		FLocalFileNetworkReplayStreamer::StartStreaming(CustomName, FriendlyName, UserStrings, bRecord, ReplayVersion, Delegate);
+		FLocalFileNetworkReplayStreamer::StartStreaming(Params, Delegate);
 	}
-	else if (UserStrings.Num() > 0 && !UserStrings[0].IsEmpty())
+	else if (Params.UserIndices.Num() > 0)
 	{
-		TArray<int32> UserIndices;
-		GetUserIndicesFromUserStrings(UserStrings, UserIndices);
-		StartStreaming(CustomName, FriendlyName, UserIndices, bRecord, ReplayVersion, Delegate);
+		StartStreamingSaved(Params, Delegate);
 	}
 	else
 	{
@@ -672,25 +652,29 @@ void FSaveGameNetworkReplayStreamer::StartStreaming(const FString& CustomName, c
 	}
 }
 
-void FSaveGameNetworkReplayStreamer::StartStreamingSaved(const FString& CustomName, const FString& FriendlyName, const TArray<int32>& UserIndices, bool bRecord, const FNetworkReplayVersion& ReplayVersion, const FStartStreamingCallback& Delegate)
+void FSaveGameNetworkReplayStreamer::StartStreamingSaved(const FStartStreamingParameters& Params, const FStartStreamingCallback& Delegate)
 {
 	// We should only hit this path if we're playing back a replay.
-	check(!bRecord);
+	check(!Params.bRecord);
 
 	using TAsyncTypes = SaveGameReplay::TAsyncTypes<FStartStreamingResult>;
 
-	auto AsyncWork = [this, CustomName, UserIndices, ReplayVersion]()
+	auto AsyncWork = [this, Params]()
 	{
 		auto SharedResult = TAsyncTypes::MakeSharedResult();
-		StartStreaming_Internal(CustomName, FString(), UserIndices, false, ReplayVersion, *SharedResult.Get());
+		StartStreaming_Internal(Params, *SharedResult.Get());
 		return SharedResult;
 	};
 
-	auto PostAsyncWork = [this, UserIndices, ReplayVersion, Delegate](const FStartStreamingResult& Result)
+	auto PostAsyncWork = [this, Params, Delegate](const FStartStreamingResult& Result)
 	{
 		if (Result.WasSuccessful())
 		{
-			FLocalFileNetworkReplayStreamer::StartStreaming(GetLocalPlaybackName(), FString(), UserIndices, false, ReplayVersion, Delegate);
+			FStartStreamingParameters LocalParams = Params;
+			LocalParams.CustomName = GetLocalPlaybackName();
+			LocalParams.FriendlyName = FString();
+
+			FLocalFileNetworkReplayStreamer::StartStreaming(LocalParams, Delegate);
 		}
 		else
 		{
@@ -701,11 +685,11 @@ void FSaveGameNetworkReplayStreamer::StartStreamingSaved(const FString& CustomNa
 	SaveGameReplay::FAsyncTaskManager::Get().StartTask<FStartStreamingResult>(*this, TEXT("StartStreaming"), AsyncWork, PostAsyncWork);
 }
 
-void FSaveGameNetworkReplayStreamer::StartStreaming_Internal(const FString& CustomName, const FString& FriendlyName, const TArray<int32>& UserIndices, bool bRecord, const FNetworkReplayVersion& ReplayVersion, FStartStreamingResult& Result)
+void FSaveGameNetworkReplayStreamer::StartStreaming_Internal(const FStartStreamingParameters& Params, FStartStreamingResult& Result)
 {
 	using ESaveExistsResult = ISaveGameSystem::ESaveExistsResult;
 
-	Result.bRecording = bRecord;
+	Result.bRecording = Params.bRecord;
 
 	// Make sure the save game system is available.
 	ISaveGameSystem* SaveGameSystem = IPlatformFeaturesModule::Get().GetSaveGameSystem();
@@ -714,16 +698,16 @@ void FSaveGameNetworkReplayStreamer::StartStreaming_Internal(const FString& Cust
 		return;
 	}
 
-	if (bRecord)
+	if (Params.bRecord)
 	{
 		UE_LOG(LogSaveGameReplay, Warning, TEXT("FSaveGameNetworkReplayStreamer::StartStreaming: Cannot record directly to a save game, use KeepReplay() instead."));
 		return;
 	}
 
-	const int32 UserIndex = UserIndices.Num() > 0 ? UserIndices[0] : INDEX_NONE;
+	const int32 UserIndex = Params.UserIndices.Num() > 0 ? Params.UserIndices[0] : INDEX_NONE;
 
 	// Make sure that the file actually exists.
-	const ESaveExistsResult SaveExistsResult = SaveGameSystem->DoesSaveGameExistWithResult(*CustomName, UserIndex);
+	const ESaveExistsResult SaveExistsResult = SaveGameSystem->DoesSaveGameExistWithResult(*Params.CustomName, UserIndex);
 	if (SaveExistsResult != ESaveExistsResult::OK)
 	{
 		PopulateStreamingResultFromSaveExistsResult(SaveExistsResult, Result);
@@ -733,7 +717,7 @@ void FSaveGameNetworkReplayStreamer::StartStreaming_Internal(const FString& Cust
 
 	// Try to load the data.
 	TArray<uint8> ReplayData;
-	if (!SaveGameSystem->LoadGame(/*bAttemptToUseUI=*/false, *CustomName, UserIndex, ReplayData))
+	if (!SaveGameSystem->LoadGame(/*bAttemptToUseUI=*/false, *Params.CustomName, UserIndex, ReplayData))
 	{
 		UE_LOG(LogSaveGameReplay, Warning, TEXT("FSaveGameNetworkReplayStreamer::StartStreaming: Failed to load replay data."));
 		return;
@@ -1171,22 +1155,7 @@ void FSaveGameNetworkReplayStreamer::RenameReplaySaved(const FString& ReplayName
 	Delegate.ExecuteIfBound(Result);
 }
 
-void FSaveGameNetworkReplayStreamer::EnumerateStreams(const FNetworkReplayVersion& ReplayVersion, const FString& UserString, const FString& MetaString, const FEnumerateStreamsCallback& Delegate)
-{
-	EnumerateStreamsSaved(ReplayVersion, GetUserIndexFromUserString(UserString), MetaString, TArray< FString >(), Delegate);
-}
-
-void FSaveGameNetworkReplayStreamer::EnumerateStreams(const FNetworkReplayVersion& ReplayVersion, const FString& UserString, const FString& MetaString, const TArray< FString >& ExtraParms, const FEnumerateStreamsCallback& Delegate)
-{
-	EnumerateStreamsSaved(ReplayVersion, GetUserIndexFromUserString(UserString), MetaString, ExtraParms, Delegate);
-}
-
 void FSaveGameNetworkReplayStreamer::EnumerateStreams(const FNetworkReplayVersion& ReplayVersion, const int32 UserIndex, const FString& MetaString, const TArray< FString >& ExtraParms, const FEnumerateStreamsCallback& Delegate)
-{
-	EnumerateStreamsSaved(ReplayVersion, UserIndex, MetaString, ExtraParms, Delegate);
-}
-
-void FSaveGameNetworkReplayStreamer::EnumerateStreamsSaved(const FNetworkReplayVersion& ReplayVersion, const int32 UserIndex, const FString& MetaString, const TArray< FString >& ExtraParms, const FEnumerateStreamsCallback& Delegate)
 {
 	using TAsyncTypes = SaveGameReplay::TAsyncTypes<FEnumerateStreamsResult>;
 
@@ -1285,13 +1254,6 @@ void FSaveGameNetworkReplayStreamer::EnumerateRecentStreams(const FNetworkReplay
 	FLocalFileNetworkReplayStreamer::EnumerateStreams(ReplayVersion, UserIndex, FString(), TArray<FString>(), Delegate);
 }
 
-void FSaveGameNetworkReplayStreamer::EnumerateRecentStreams(const FNetworkReplayVersion& ReplayVersion, const FString& RecentViewer, const FEnumerateStreamsCallback& Delegate)
-{
-	// Recent Streams will just be any stream we have locally that hasn't been committed to memory.
-	// So, just do Local Stream enumeration.
-	FLocalFileNetworkReplayStreamer::EnumerateStreams(ReplayVersion, RecentViewer, FString(), TArray<FString>(), Delegate);
-}
-
 void FSaveGameNetworkReplayStreamer::EnumerateEvents(const FString& ReplayName, const FString& Group, const FEnumerateEventsCallback& Delegate)
 {
 	if (!IsSaveGameFileName(ReplayName))
diff --git a/Engine/Source/Runtime/NetworkReplayStreaming/SaveGameNetworkReplayStreaming/Public/SaveGameNetworkReplayStreaming.h b/Engine/Source/Runtime/NetworkReplayStreaming/SaveGameNetworkReplayStreaming/Public/SaveGameNetworkReplayStreaming.h
--- a/Engine/Source/Runtime/NetworkReplayStreaming/SaveGameNetworkReplayStreaming/Public/SaveGameNetworkReplayStreaming.h
+++ b/Engine/Source/Runtime/NetworkReplayStreaming/SaveGameNetworkReplayStreaming/Public/SaveGameNetworkReplayStreaming.h
@@ -37,14 +37,10 @@ public:
 	FSaveGameNetworkReplayStreamer(const FString& DemoSavePath, const FString& PlaybackReplayName);
 
 	/** INetworkReplayStreamer implementation */
-	virtual void StartStreaming(const FString& CustomName, const FString& FriendlyName, const TArray<FString>& UserNames, bool bRecord, const FNetworkReplayVersion& ReplayVersion, const FStartStreamingCallback& Delegate) override;
-	virtual void StartStreaming(const FString& CustomName, const FString& FriendlyName, const TArray<int32>& UserIndices, bool bRecord, const FNetworkReplayVersion& ReplayVersion, const FStartStreamingCallback& Delegate) override;
+	virtual void StartStreaming(const FStartStreamingParameters& Params, const FStartStreamingCallback& Delegate) override;
 	virtual void DeleteFinishedStream(const FString& ReplayName, const FDeleteFinishedStreamCallback& Delegate) override;
 	virtual void DeleteFinishedStream(const FString& StreamName, const int32 UserIndex, const FDeleteFinishedStreamCallback& Delegate) override;
-	virtual void EnumerateStreams(const FNetworkReplayVersion& ReplayVersion, const FString& UserString, const FString& MetaString, const FEnumerateStreamsCallback& Delegate) override;
-	virtual void EnumerateStreams(const FNetworkReplayVersion& ReplayVersion, const FString& UserString, const FString& MetaString, const TArray< FString >& ExtraParms, const FEnumerateStreamsCallback& Delegate) override;
 	virtual void EnumerateStreams(const FNetworkReplayVersion& InReplayVersion, const int32 UserIndex, const FString& MetaString, const TArray< FString >& ExtraParms, const FEnumerateStreamsCallback& Delegate) override;
-	virtual void EnumerateRecentStreams(const FNetworkReplayVersion& ReplayVersion, const FString& RecentViewer, const FEnumerateStreamsCallback& Delegate) override;
 	virtual void EnumerateRecentStreams(const FNetworkReplayVersion& ReplayVersion, const int32 UserIndex, const FEnumerateStreamsCallback& Delegate) override;
 	virtual void KeepReplay(const FString& ReplayName, const bool bKeep, const FKeepReplayCallback& Delegate) override;
 	virtual void KeepReplay(const FString& ReplayName, const bool bKeep, const int32 UserIndex, const FKeepReplayCallback& Delegate) override;
@@ -63,9 +59,8 @@ public:
 	virtual void RequestEventGroupData(const FString& ReplayName, const FString& Group, const int32 UserIndex, const FRequestEventGroupDataCallback& Delegate) override;
 protected:
 
-	void StartStreamingSaved(const FString& CustomName, const FString& FriendlyName, const TArray<int32>& UserIndices, bool bRecord, const FNetworkReplayVersion& ReplayVersion, const FStartStreamingCallback& Delegate);
+	void StartStreamingSaved(const FStartStreamingParameters& Params, const FStartStreamingCallback& Delegate);
 	void DeleteFinishedStreamSaved(const FString& ReplayName, const int32 UserIndex, const FDeleteFinishedStreamCallback& Delegate) const;
-	void EnumerateStreamsSaved(const FNetworkReplayVersion& ReplayVersion, const int32 UserIndex, const FString& MetaString, const TArray< FString >& ExtraParms, const FEnumerateStreamsCallback& Result);
 	void KeepReplaySaved(const FString& ReplayName, const bool bKeep, const int32 UserIndex, const FKeepReplayCallback& Result);
 	void RenameReplayFriendlyNameSaved(const FString& ReplayName, const FString& NewFriendlyName, const int32 UserIndex, const FRenameReplayCallback& Delegate) const;
 	void RenameReplaySaved(const FString& ReplayName, const FString& NewName, const int32 UserIndex, const FRenameReplayCallback& Delegate);
@@ -73,7 +68,7 @@ protected:
 	void RequestEventDataSaved(const FString& ReplayName, const FString& EventID, const int32 UserIndex, const FRequestEventDataCallback& Delegate);
 	void RequestEventGroupDataSaved(const FString& ReplayName, const FString& Group, const int32 UserIndex, const FRequestEventGroupDataCallback& Delegate);
 
-	void StartStreaming_Internal(const FString& CustomName, const FString& FriendlyName, const TArray<int32>& UserIndices, bool bRecord, const FNetworkReplayVersion& ReplayVersion, FStartStreamingResult& Result);
+	void StartStreaming_Internal(const FStartStreamingParameters& Params, FStartStreamingResult& Result);
 	void DeleteFinishedStream_Internal(const FString& ReplayName, const int32 UserIndex, FDeleteFinishedStreamResult& Result) const;
 	void EnumerateStreams_Internal(const FNetworkReplayVersion& ReplayVersion, const int32 UserIndex, const FString& MetaString, const TArray< FString >& ExtraParms, FEnumerateStreamsResult& Result);
 	void KeepReplay_Internal(const FString& ReplayName, const bool bKeep, const int32 UserIndex, FKeepReplayResult& Result);
