Compile fixes for in memory replay streamer

diff --git a/Engine/Source/Runtime/NetworkReplayStreaming/InMemoryNetworkReplayStreaming/Private/InMemoryNetworkReplayStreaming.cpp b/Engine/Source/Runtime/NetworkReplayStreaming/InMemoryNetworkReplayStreaming/Private/InMemoryNetworkReplayStreaming.cpp
--- a/Engine/Source/Runtime/NetworkReplayStreaming/InMemoryNetworkReplayStreaming/Private/InMemoryNetworkReplayStreaming.cpp
+++ b/Engine/Source/Runtime/NetworkReplayStreaming/InMemoryNetworkReplayStreaming/Private/InMemoryNetworkReplayStreaming.cpp
@@ -13,19 +13,14 @@ static FString GetAutomaticDemoName()
 	return FGuid::NewGuid().ToString();
 }
 
-void FInMemoryNetworkReplayStreamer::StartStreaming(const FString& CustomName, const FString& FriendlyName, const TArray< int32 >& UserIndices, bool bRecord, const FNetworkReplayVersion& ReplayVersion, const FStartStreamingCallback& Delegate)
-{
-	StartStreaming(CustomName, FriendlyName, TArray<FString>(), bRecord, ReplayVersion, Delegate);
-}
-
-void FInMemoryNetworkReplayStreamer::StartStreaming( const FString& CustomName, const FString& FriendlyName, const TArray< FString >& UserNames, bool bRecord, const FNetworkReplayVersion& ReplayVersion, const FStartStreamingCallback& Delegate )
+void FInMemoryNetworkReplayStreamer::StartStreaming(const FStartStreamingParameters& Params, const FStartStreamingCallback& Delegate)
 {
 	FStartStreamingResult Result;
-	Result.bRecording = bRecord;
+	Result.bRecording = Params.bRecord;
 
-	if ( CustomName.IsEmpty() )
+	if ( Params.CustomName.IsEmpty() )
 	{
-		if ( bRecord )
+		if ( Params.bRecord )
 		{
 			// If we're recording and the caller didn't provide a name, generate one automatically
 			CurrentStreamName = GetAutomaticDemoName();
@@ -40,10 +35,10 @@ void FInMemoryNetworkReplayStreamer::StartStreaming( const FString& CustomName,
 	}
 	else
 	{
-		CurrentStreamName = CustomName;
+		CurrentStreamName = Params.CustomName;
 	}
 
-	if ( !bRecord )
+	if ( !Params.bRecord )
 	{
 		FInMemoryReplay* FoundReplay = GetCurrentReplay();
 		if (FoundReplay == nullptr)
@@ -54,8 +49,8 @@ void FInMemoryNetworkReplayStreamer::StartStreaming( const FString& CustomName,
 		}
 
 		FileAr.Reset(new FInMemoryReplayStreamArchive(FoundReplay->StreamChunks));
-		FileAr->SetIsSaving(bRecord);
-		FileAr->SetIsLoading(!bRecord);
+		FileAr->SetIsSaving(Params.bRecord);
+		FileAr->SetIsLoading(!Params.bRecord);
 		HeaderAr.Reset(new FMemoryReader(FoundReplay->Header));
 		StreamerState = EStreamerState::Playback;
 	}
@@ -65,16 +60,16 @@ void FInMemoryNetworkReplayStreamer::StartStreaming( const FString& CustomName,
 		TUniquePtr<FInMemoryReplay> NewReplay(new FInMemoryReplay);
 
 		NewReplay->StreamInfo.Name = CurrentStreamName;
-		NewReplay->StreamInfo.FriendlyName = FriendlyName;
+		NewReplay->StreamInfo.FriendlyName = Params.FriendlyName;
 		NewReplay->StreamInfo.Timestamp = FDateTime::Now();
 		NewReplay->StreamInfo.bIsLive = true;
-		NewReplay->StreamInfo.Changelist = ReplayVersion.Changelist;
-		NewReplay->NetworkVersion = ReplayVersion.NetworkVersion;
+		NewReplay->StreamInfo.Changelist = Params.ReplayVersion.Changelist;
+		NewReplay->NetworkVersion = Params.ReplayVersion.NetworkVersion;
 
 		// Open archives for writing
 		FileAr.Reset(new FInMemoryReplayStreamArchive(NewReplay->StreamChunks));
-		FileAr->SetIsSaving(bRecord);
-		FileAr->SetIsLoading(!bRecord);
+		FileAr->SetIsSaving(Params.bRecord);
+		FileAr->SetIsLoading(!Params.bRecord);
 		HeaderAr.Reset(new FMemoryWriter(NewReplay->Header));
 
 		OwningFactory->Replays.Add(CurrentStreamName, MoveTemp(NewReplay));
@@ -193,11 +188,6 @@ void FInMemoryNetworkReplayStreamer::DeleteFinishedStream( const FString& Stream
 }
 
 void FInMemoryNetworkReplayStreamer::EnumerateRecentStreams(const FNetworkReplayVersion& ReplayVersion, const int32 UserIndex, const FEnumerateStreamsCallback& Delegate)
-{
-	EnumerateRecentStreams(ReplayVersion, FString(), Delegate);
-}
-
-void FInMemoryNetworkReplayStreamer::EnumerateRecentStreams(const FNetworkReplayVersion& ReplayVersion, const FString& RecentViewer, const FEnumerateStreamsCallback& Delegate)
 {
 	UE_LOG(LogMemoryReplay, Log, TEXT("FInMemoryNetworkReplayStreamer::EnumerateRecentStreams is currently unsupported."));
 	FEnumerateStreamsResult Result;
@@ -206,16 +196,6 @@ void FInMemoryNetworkReplayStreamer::EnumerateRecentStreams(const FNetworkReplay
 }
 
 void FInMemoryNetworkReplayStreamer::EnumerateStreams(const FNetworkReplayVersion& ReplayVersion, const int32 UserIndex, const FString& MetaString, const TArray< FString >& ExtraParms, const FEnumerateStreamsCallback& Delegate)
-{
-	EnumerateStreams(ReplayVersion, FString(), MetaString, ExtraParms, Delegate);
-}
-
-void FInMemoryNetworkReplayStreamer::EnumerateStreams( const FNetworkReplayVersion& ReplayVersion, const FString& UserString, const FString& MetaString, const FEnumerateStreamsCallback& Delegate )
-{
-	EnumerateStreams( ReplayVersion, UserString, MetaString, TArray< FString >(), Delegate );
-}
-
-void FInMemoryNetworkReplayStreamer::EnumerateStreams( const FNetworkReplayVersion& ReplayVersion, const FString& UserString, const FString& MetaString, const TArray< FString >& ExtraParms, const FEnumerateStreamsCallback& Delegate )
 {
 	FEnumerateStreamsResult Result;
 	Result.Result = EStreamingOperationResult::Success;
diff --git a/Engine/Source/Runtime/NetworkReplayStreaming/InMemoryNetworkReplayStreaming/Public/InMemoryNetworkReplayStreaming.h b/Engine/Source/Runtime/NetworkReplayStreaming/InMemoryNetworkReplayStreaming/Public/InMemoryNetworkReplayStreaming.h
--- a/Engine/Source/Runtime/NetworkReplayStreaming/InMemoryNetworkReplayStreaming/Public/InMemoryNetworkReplayStreaming.h
+++ b/Engine/Source/Runtime/NetworkReplayStreaming/InMemoryNetworkReplayStreaming/Public/InMemoryNetworkReplayStreaming.h
@@ -141,10 +141,7 @@ public:
 	virtual bool IsLive() const override;
 	virtual void DeleteFinishedStream( const FString& StreamName, const FDeleteFinishedStreamCallback& Delegate) override;
 	virtual void DeleteFinishedStream( const FString& StreamName, const int32 UserIndex, const FDeleteFinishedStreamCallback& Delegate ) override;
-	virtual void EnumerateStreams( const FNetworkReplayVersion& ReplayVersion, const FString& UserString, const FString& MetaString, const FEnumerateStreamsCallback& Delegate ) override;
-	virtual void EnumerateStreams( const FNetworkReplayVersion& InReplayVersion, const FString& UserString, const FString& MetaString, const TArray< FString >& ExtraParms, const FEnumerateStreamsCallback& Delegate ) override;
-	virtual void EnumerateStreams( const FNetworkReplayVersion& InReplayVersion, const int32 UserIndex, const FString& MetaString, const TArray< FString >& ExtraParms, const FEnumerateStreamsCallback& Delegate ) override;
-	virtual void EnumerateRecentStreams( const FNetworkReplayVersion& ReplayVersion, const FString& RecentViewer, const FEnumerateStreamsCallback& Delegate ) override;
+	virtual void EnumerateStreams( const FNetworkReplayVersion& ReplayVersion, const int32 UserIndex, const FString& MetaString, const TArray< FString >& ExtraParms, const FEnumerateStreamsCallback& Delegate ) override;
 	virtual void EnumerateRecentStreams( const FNetworkReplayVersion& ReplayVersion, const int32 UserIndex, const FEnumerateStreamsCallback& Delegate ) override;
 	virtual ENetworkReplayError::Type GetLastError() const override { return ENetworkReplayError::None; }
 	virtual void AddUserToReplay(const FString& UserString) override;
