Non-unity include fix

diff --git a/Engine/Source/Runtime/NetworkReplayStreaming/NetworkReplayStreaming/Private/NetworkReplayStreaming.cpp b/Engine/Source/Runtime/NetworkReplayStreaming/NetworkReplayStreaming/Private/NetworkReplayStreaming.cpp
--- a/Engine/Source/Runtime/NetworkReplayStreaming/NetworkReplayStreaming/Private/NetworkReplayStreaming.cpp
+++ b/Engine/Source/Runtime/NetworkReplayStreaming/NetworkReplayStreaming/Private/NetworkReplayStreaming.cpp
@@ -4,6 +4,9 @@
 #include "Misc/CommandLine.h"
 #include "Misc/ConfigCacheIni.h"
 #include "HAL/ConsoleManager.h"
+#include "Engine/Engine.h"
+#include "Engine/World.h"
+#include "Engine/LocalPlayer.h"
 
 IMPLEMENT_MODULE( FNetworkReplayStreaming, NetworkReplayStreaming );
 
