Fixing circular dependency with engine and replay streamer interface

diff --git a/Engine/Source/Runtime/Engine/Private/DemoNetDriver.cpp b/Engine/Source/Runtime/Engine/Private/DemoNetDriver.cpp
--- a/Engine/Source/Runtime/Engine/Private/DemoNetDriver.cpp
+++ b/Engine/Source/Runtime/Engine/Private/DemoNetDriver.cpp
@@ -905,7 +905,7 @@ bool UDemoNetDriver::InitConnect(FNetworkNotify* InNotify, const FURL& ConnectUR
 	
 	FStartStreamingParameters Params;
 	Params.CustomName = DemoURL.Map;
-	Params.DemoURL = DemoURL;
+	Params.DemoURL = GetDemoURL();
 	Params.UserIndices = MoveTemp(UserIndices);
 	Params.bRecord = false;
 	Params.ReplayVersion = FNetworkVersion::GetReplayVersion();
@@ -1093,7 +1093,7 @@ bool UDemoNetDriver::InitListen(FNetworkNotify* InNotify, FURL& ListenURL, bool
 	FStartStreamingParameters Params;
 	Params.CustomName = DemoURL.Map;
 	Params.FriendlyName = FriendlyNameOption != nullptr ? FString(FriendlyNameOption) : World->GetMapName();
-	Params.DemoURL = DemoURL;
+	Params.DemoURL = GetDemoURL();
 	Params.UserIndices = MoveTemp(UserIndices);
 	Params.bRecord = true;
 	Params.ReplayVersion = FNetworkVersion::GetReplayVersion();
diff --git a/Engine/Source/Runtime/NetworkReplayStreaming/HttpNetworkReplayStreaming/Private/HttpNetworkReplayStreaming.cpp b/Engine/Source/Runtime/NetworkReplayStreaming/HttpNetworkReplayStreaming/Private/HttpNetworkReplayStreaming.cpp
--- a/Engine/Source/Runtime/NetworkReplayStreaming/HttpNetworkReplayStreaming/Private/HttpNetworkReplayStreaming.cpp
+++ b/Engine/Source/Runtime/NetworkReplayStreaming/HttpNetworkReplayStreaming/Private/HttpNetworkReplayStreaming.cpp
@@ -2702,6 +2702,28 @@ bool FHttpNetworkReplayStreamer::IsCheckpointTypeSupported(EReplayCheckpointType
 	return bSupported;
 }
 
+const int32 FHttpNetworkReplayStreamer::GetUserIndexFromUserString(const FString& UserString)
+{
+	if (!UserString.IsEmpty() && GEngine != nullptr)
+	{
+		if (UWorld* World = GWorld.GetReference())
+		{
+			for (auto ConstIt = GEngine->GetLocalPlayerIterator(World); ConstIt; ++ConstIt)
+			{
+				if (ULocalPlayer const * const LocalPlayer = *ConstIt)
+				{
+					if (UserString.Equals(LocalPlayer->GetPreferredUniqueNetId().ToString()))
+					{
+						return LocalPlayer->GetControllerId();
+					}
+				}
+			}
+		}
+	}
+
+	return INDEX_NONE;
+}
+
 IMPLEMENT_MODULE( FHttpNetworkReplayStreamingFactory, HttpNetworkReplayStreaming )
 
 TSharedPtr< INetworkReplayStreamer > FHttpNetworkReplayStreamingFactory::CreateReplayStreamer()
diff --git a/Engine/Source/Runtime/NetworkReplayStreaming/HttpNetworkReplayStreaming/Public/HttpNetworkReplayStreaming.h b/Engine/Source/Runtime/NetworkReplayStreaming/HttpNetworkReplayStreaming/Public/HttpNetworkReplayStreaming.h
--- a/Engine/Source/Runtime/NetworkReplayStreaming/HttpNetworkReplayStreaming/Public/HttpNetworkReplayStreaming.h
+++ b/Engine/Source/Runtime/NetworkReplayStreaming/HttpNetworkReplayStreaming/Public/HttpNetworkReplayStreaming.h
@@ -261,6 +261,8 @@ public:
 
 	virtual bool IsCheckpointTypeSupported(EReplayCheckpointType CheckpointType) const override;
 
+	virtual const int32 GetUserIndexFromUserString(const FString& UserString) override;
+
 	/** FHttpNetworkReplayStreamer */
 	void UploadHeader();
 	void FlushStream();
diff --git a/Engine/Source/Runtime/NetworkReplayStreaming/InMemoryNetworkReplayStreaming/Private/InMemoryNetworkReplayStreaming.cpp b/Engine/Source/Runtime/NetworkReplayStreaming/InMemoryNetworkReplayStreaming/Private/InMemoryNetworkReplayStreaming.cpp
--- a/Engine/Source/Runtime/NetworkReplayStreaming/InMemoryNetworkReplayStreaming/Private/InMemoryNetworkReplayStreaming.cpp
+++ b/Engine/Source/Runtime/NetworkReplayStreaming/InMemoryNetworkReplayStreaming/Private/InMemoryNetworkReplayStreaming.cpp
@@ -5,6 +5,9 @@
 #include "Misc/Guid.h"
 #include "Misc/DateTime.h"
 #include "HAL/IConsoleManager.h"
+#include "Engine/Engine.h"
+#include "Engine/World.h"
+#include "Engine/LocalPlayer.h"
 
 DEFINE_LOG_CATEGORY_STATIC( LogMemoryReplay, Log, All );
 
@@ -557,6 +560,28 @@ TStatId FInMemoryNetworkReplayStreamer::GetStatId() const
 	RETURN_QUICK_DECLARE_CYCLE_STAT(FInMemoryNetworkReplayStreamer, STATGROUP_Tickables);
 }
 
+const int32 FInMemoryNetworkReplayStreamer::GetUserIndexFromUserString(const FString& UserString)
+{
+	if (!UserString.IsEmpty() && GEngine != nullptr)
+	{
+		if (UWorld* World = GWorld.GetReference())
+		{
+			for (auto ConstIt = GEngine->GetLocalPlayerIterator(World); ConstIt; ++ConstIt)
+			{
+				if (ULocalPlayer const * const LocalPlayer = *ConstIt)
+				{
+					if (UserString.Equals(LocalPlayer->GetPreferredUniqueNetId().ToString()))
+					{
+						return LocalPlayer->GetControllerId();
+					}
+				}
+			}
+		}
+	}
+
+	return INDEX_NONE;
+}
+
 void FInMemoryReplayStreamArchive::Serialize(void* V, int64 Length) 
 {
 	if (IsLoading() )
diff --git a/Engine/Source/Runtime/NetworkReplayStreaming/InMemoryNetworkReplayStreaming/Public/InMemoryNetworkReplayStreaming.h b/Engine/Source/Runtime/NetworkReplayStreaming/InMemoryNetworkReplayStreaming/Public/InMemoryNetworkReplayStreaming.h
--- a/Engine/Source/Runtime/NetworkReplayStreaming/InMemoryNetworkReplayStreaming/Public/InMemoryNetworkReplayStreaming.h
+++ b/Engine/Source/Runtime/NetworkReplayStreaming/InMemoryNetworkReplayStreaming/Public/InMemoryNetworkReplayStreaming.h
@@ -194,6 +194,8 @@ public:
 
 	virtual bool IsCheckpointTypeSupported(EReplayCheckpointType CheckpointType) const override { return (CheckpointType == EReplayCheckpointType::Full); }
 
+	virtual const int32 GetUserIndexFromUserString(const FString& UserString) override;
+
 private:
 	bool IsNamedStreamLive( const FString& StreamName ) const;
 
diff --git a/Engine/Source/Runtime/NetworkReplayStreaming/LocalFileNetworkReplayStreaming/Private/LocalFileNetworkReplayStreaming.cpp b/Engine/Source/Runtime/NetworkReplayStreaming/LocalFileNetworkReplayStreaming/Private/LocalFileNetworkReplayStreaming.cpp
--- a/Engine/Source/Runtime/NetworkReplayStreaming/LocalFileNetworkReplayStreaming/Private/LocalFileNetworkReplayStreaming.cpp
+++ b/Engine/Source/Runtime/NetworkReplayStreaming/LocalFileNetworkReplayStreaming/Private/LocalFileNetworkReplayStreaming.cpp
@@ -2964,6 +2964,28 @@ bool FLocalFileNetworkReplayStreamer::IsCheckpointTypeSupported(EReplayCheckpoin
 	return bSupported;
 }
 
+const int32 FLocalFileNetworkReplayStreamer::GetUserIndexFromUserString(const FString& UserString)
+{
+	if (!UserString.IsEmpty() && GEngine != nullptr)
+	{
+		if (UWorld* World = GWorld.GetReference())
+		{
+			for (auto ConstIt = GEngine->GetLocalPlayerIterator(World); ConstIt; ++ConstIt)
+			{
+				if (ULocalPlayer const * const LocalPlayer = *ConstIt)
+				{
+					if (UserString.Equals(LocalPlayer->GetPreferredUniqueNetId().ToString()))
+					{
+						return LocalPlayer->GetControllerId();
+					}
+				}
+			}
+		}
+	}
+
+	return INDEX_NONE;
+}
+
 IMPLEMENT_MODULE(FLocalFileNetworkReplayStreamingFactory, LocalFileNetworkReplayStreaming)
 
 TSharedPtr<INetworkReplayStreamer> FLocalFileNetworkReplayStreamingFactory::CreateReplayStreamer() 
diff --git a/Engine/Source/Runtime/NetworkReplayStreaming/LocalFileNetworkReplayStreaming/Public/LocalFileNetworkReplayStreaming.h b/Engine/Source/Runtime/NetworkReplayStreaming/LocalFileNetworkReplayStreaming/Public/LocalFileNetworkReplayStreaming.h
--- a/Engine/Source/Runtime/NetworkReplayStreaming/LocalFileNetworkReplayStreaming/Public/LocalFileNetworkReplayStreaming.h
+++ b/Engine/Source/Runtime/NetworkReplayStreaming/LocalFileNetworkReplayStreaming/Public/LocalFileNetworkReplayStreaming.h
@@ -445,6 +445,8 @@ public:
 
 	virtual bool IsCheckpointTypeSupported(EReplayCheckpointType CheckpointType) const override;
 
+	virtual const int32 GetUserIndexFromUserString(const FString& UserString) override;
+
 	virtual bool SupportsCompression() const { return false; }
 	virtual int32 GetDecompressedSize(FArchive& InCompressed) const { return 0; }
 	virtual bool DecompressBuffer(const TArray<uint8>& InCompressed, TArray< uint8 >& OutBuffer) const { return false; }
diff --git a/Engine/Source/Runtime/NetworkReplayStreaming/NetworkReplayStreaming/NetworkReplayStreaming.Build.cs b/Engine/Source/Runtime/NetworkReplayStreaming/NetworkReplayStreaming/NetworkReplayStreaming.Build.cs
--- a/Engine/Source/Runtime/NetworkReplayStreaming/NetworkReplayStreaming/NetworkReplayStreaming.Build.cs
+++ b/Engine/Source/Runtime/NetworkReplayStreaming/NetworkReplayStreaming/NetworkReplayStreaming.Build.cs
@@ -12,7 +12,6 @@ namespace UnrealBuildTool.Rules
 				new string[]
 				{
 					"Core",
-					"Engine",
                     "Json",
                     // ... add other public dependencies that you statically link with here ...
 				}
diff --git a/Engine/Source/Runtime/NetworkReplayStreaming/NetworkReplayStreaming/Private/NetworkReplayStreaming.cpp b/Engine/Source/Runtime/NetworkReplayStreaming/NetworkReplayStreaming/Private/NetworkReplayStreaming.cpp
--- a/Engine/Source/Runtime/NetworkReplayStreaming/NetworkReplayStreaming/Private/NetworkReplayStreaming.cpp
+++ b/Engine/Source/Runtime/NetworkReplayStreaming/NetworkReplayStreaming/Private/NetworkReplayStreaming.cpp
@@ -4,9 +4,6 @@
 #include "Misc/CommandLine.h"
 #include "Misc/ConfigCacheIni.h"
 #include "HAL/ConsoleManager.h"
-#include "Engine/Engine.h"
-#include "Engine/World.h"
-#include "Engine/LocalPlayer.h"
 
 IMPLEMENT_MODULE( FNetworkReplayStreaming, NetworkReplayStreaming );
 
@@ -132,7 +129,12 @@ void INetworkReplayStreamer::StartStreaming(const FString& CustomName, const FSt
 	FStartStreamingParameters Params;
 	Params.CustomName = CustomName;
 	Params.FriendlyName = FriendlyName;
-	GetUserIndicesFromUserStrings(UserNames, Params.UserIndices);
+	
+	for (const FString& UserString : UserNames)
+	{
+		Params.UserIndices.Emplace(GetUserIndexFromUserString(UserString));
+	}
+	
 	Params.bRecord = bRecord;
 	Params.ReplayVersion = ReplayVersion;
 
@@ -164,59 +166,4 @@ void INetworkReplayStreamer::EnumerateStreams(const FNetworkReplayVersion& Repla
 void INetworkReplayStreamer::EnumerateRecentStreams(const FNetworkReplayVersion& ReplayVersion, const FString& RecentViewer, const FEnumerateStreamsCallback& Delegate)
 {
 	EnumerateRecentStreams(ReplayVersion, GetUserIndexFromUserString(RecentViewer), Delegate);
-}
-
-const int32 INetworkReplayStreamer::GetUserIndexFromUserString(const FString& UserString)
-{
-	if (!UserString.IsEmpty() && GEngine != nullptr)
-	{
-		if (UWorld* World = GWorld.GetReference())
-		{
-			for (auto ConstIt = GEngine->GetLocalPlayerIterator(World); ConstIt; ++ConstIt)
-			{
-				if (ULocalPlayer const * const LocalPlayer = *ConstIt)
-				{
-					if (UserString.Equals(LocalPlayer->GetPreferredUniqueNetId().ToString()))
-					{
-						return LocalPlayer->GetControllerId();
-					}
-				}
-			}
-		}
-	}
-
-	return INDEX_NONE;
-}
-
-const void INetworkReplayStreamer::GetUserIndicesFromUserStrings(const TArray<FString>& UserStrings, TArray<int32>& OutUserIndices)
-{
-	if (GEngine != nullptr)
-	{
-		if (UserStrings.Num() == 1)
-		{
-			OutUserIndices.Emplace(GetUserIndexFromUserString(UserStrings[0]));
-		}
-		else if (UserStrings.Num() > 1)
-		{
-			if (UWorld* World = GWorld.GetReference())
-			{
-				TMap<FString, int32> StringToId;
-				for (auto ConstIt = GEngine->GetLocalPlayerIterator(World); ConstIt; ++ConstIt)
-				{
-					if (ULocalPlayer const * const LocalPlayer = *ConstIt)
-					{
-						StringToId.Emplace(LocalPlayer->GetPreferredUniqueNetId().ToString(), LocalPlayer->GetControllerId());
-					}
-				}
-
-				for (const FString& UserString : UserStrings)
-				{
-					if (int32 const * const UserIndex = StringToId.Find(UserString))
-					{
-						OutUserIndices.Emplace(*UserIndex);
-					}
-				}
-			}
-		}
-	}
 }
\ No newline at end of file
diff --git a/Engine/Source/Runtime/NetworkReplayStreaming/NetworkReplayStreaming/Public/NetworkReplayStreaming.h b/Engine/Source/Runtime/NetworkReplayStreaming/NetworkReplayStreaming/Public/NetworkReplayStreaming.h
--- a/Engine/Source/Runtime/NetworkReplayStreaming/NetworkReplayStreaming/Public/NetworkReplayStreaming.h
+++ b/Engine/Source/Runtime/NetworkReplayStreaming/NetworkReplayStreaming/Public/NetworkReplayStreaming.h
@@ -5,13 +5,11 @@
 // Dependencies.
 
 #include "CoreMinimal.h"
-#include "Engine/EngineBaseTypes.h"
 #include "Misc/NetworkVersion.h"
 #include "Modules/ModuleInterface.h"
 #include "Modules/ModuleManager.h"
 #include "Serialization/JsonSerializerMacros.h"
 
-class FNetworkReplayVersion;
 class IAnalyticsProvider;
 
 class FReplayEventListItem : public FJsonSerializable
@@ -469,7 +467,7 @@ struct FStartStreamingParameters
 {
 	FString CustomName;
 	FString FriendlyName;
-	FURL DemoURL;
+	FString DemoURL;
 	TArray<int32> UserIndices;
 	FNetworkReplayVersion ReplayVersion;
 	bool bRecord;
@@ -647,9 +645,12 @@ public:
 
 	virtual void Exec(const TCHAR* Cmd, FOutputDevice& Ar) {}
 
+protected:
+	/** Temporary to assist with deprecation of user string apis */
+	virtual const int32 GetUserIndexFromUserString(const FString& UserString) = 0;
+
 private:
-	static const int32 GetUserIndexFromUserString(const FString& UserString);
-	static const void GetUserIndicesFromUserStrings(const TArray<FString>& UserStrings, TArray<int32>& OutUserIndices);
+	const void GetUserIndicesFromUserStrings(const TArray<FString>& UserStrings, TArray<int32>& OutUserIndices);
 };
 
 /** Replay streamer factory */
diff --git a/Engine/Source/Runtime/NetworkReplayStreaming/NullNetworkReplayStreaming/Private/NullNetworkReplayStreaming.cpp b/Engine/Source/Runtime/NetworkReplayStreaming/NullNetworkReplayStreaming/Private/NullNetworkReplayStreaming.cpp
--- a/Engine/Source/Runtime/NetworkReplayStreaming/NullNetworkReplayStreaming/Private/NullNetworkReplayStreaming.cpp
+++ b/Engine/Source/Runtime/NetworkReplayStreaming/NullNetworkReplayStreaming/Private/NullNetworkReplayStreaming.cpp
@@ -4,6 +4,9 @@
 #include "HAL/FileManager.h"
 #include "Misc/Paths.h"
 #include "Misc/EngineVersion.h"
+#include "Engine/Engine.h"
+#include "Engine/World.h"
+#include "Engine/LocalPlayer.h"
 
 DEFINE_LOG_CATEGORY_STATIC( LogNullReplay, Log, All );
 
@@ -749,6 +752,28 @@ EStreamingOperationResult FNullNetworkReplayStreamer::GetDemoPath(FString& DemoP
 	return EStreamingOperationResult::Success;
 }
 
+const int32 FNullNetworkReplayStreamer::GetUserIndexFromUserString(const FString& UserString)
+{
+	if (!UserString.IsEmpty() && GEngine != nullptr)
+	{
+		if (UWorld* World = GWorld.GetReference())
+		{
+			for (auto ConstIt = GEngine->GetLocalPlayerIterator(World); ConstIt; ++ConstIt)
+			{
+				if (ULocalPlayer const * const LocalPlayer = *ConstIt)
+				{
+					if (UserString.Equals(LocalPlayer->GetPreferredUniqueNetId().ToString()))
+					{
+						return LocalPlayer->GetControllerId();
+					}
+				}
+			}
+		}
+	}
+
+	return INDEX_NONE;
+}
+
 IMPLEMENT_MODULE( FNullNetworkReplayStreamingFactory, NullNetworkReplayStreaming )
 
 TSharedPtr< INetworkReplayStreamer > FNullNetworkReplayStreamingFactory::CreateReplayStreamer() 
diff --git a/Engine/Source/Runtime/NetworkReplayStreaming/NullNetworkReplayStreaming/Public/NullNetworkReplayStreaming.h b/Engine/Source/Runtime/NetworkReplayStreaming/NullNetworkReplayStreaming/Public/NullNetworkReplayStreaming.h
--- a/Engine/Source/Runtime/NetworkReplayStreaming/NullNetworkReplayStreaming/Public/NullNetworkReplayStreaming.h
+++ b/Engine/Source/Runtime/NetworkReplayStreaming/NullNetworkReplayStreaming/Public/NullNetworkReplayStreaming.h
@@ -115,6 +115,8 @@ public:
 
 	virtual bool IsCheckpointTypeSupported(EReplayCheckpointType CheckpointType) const override { return (CheckpointType == EReplayCheckpointType::Full); }
 
+	virtual const int32 GetUserIndexFromUserString(const FString& UserString) override;
+
 private:
 	bool IsNamedStreamLive( const FString& StreamName ) const;
 
