Deprecation fixes

diff --git a/Engine/Source/Runtime/NetworkReplayStreaming/HttpNetworkReplayStreaming/Private/HttpNetworkReplayStreaming.cpp b/Engine/Source/Runtime/NetworkReplayStreaming/HttpNetworkReplayStreaming/Private/HttpNetworkReplayStreaming.cpp
--- a/Engine/Source/Runtime/NetworkReplayStreaming/HttpNetworkReplayStreaming/Private/HttpNetworkReplayStreaming.cpp
+++ b/Engine/Source/Runtime/NetworkReplayStreaming/HttpNetworkReplayStreaming/Private/HttpNetworkReplayStreaming.cpp
@@ -1516,6 +1516,18 @@ void FHttpNetworkReplayStreamer::DeleteFinishedStream( const FString& StreamName
 }
 
 void FHttpNetworkReplayStreamer::EnumerateStreams( const FNetworkReplayVersion& InReplayVersion, const int32 UserIndex, const FString& MetaString, const TArray< FString >& ExtraParms, const FEnumerateStreamsCallback& Delegate )
+{
+	EnumerateStreams( InReplayVersion, GetUserStringFromUserIndex(UserIndex), MetaString, ExtraParms, Delegate );
+}
+
+PRAGMA_DISABLE_DEPRECATION_WARNINGS
+
+void FHttpNetworkReplayStreamer::EnumerateStreams( const FNetworkReplayVersion& InReplayVersion, const FString& UserString, const FString& MetaString, const FEnumerateStreamsCallback& Delegate )
+{
+	EnumerateStreams( InReplayVersion, UserString, MetaString, TArray< FString >(), Delegate );
+}
+
+void FHttpNetworkReplayStreamer::EnumerateStreams( const FNetworkReplayVersion& InReplayVersion, const FString& UserString, const FString& MetaString, const TArray< FString >& ExtraParms, const FEnumerateStreamsCallback& Delegate )
 {
 	if ( ServerURL.IsEmpty() )
 	{
@@ -1549,7 +1561,6 @@ void FHttpNetworkReplayStreamer::EnumerateStreams( const FNetworkReplayVersion&
 	}
 
 	// Add optional User parameter (filter replays by a user that was in the replay)
-	FString UserString = GetUserStringFromUserIndex(UserIndex);
 	if ( !UserString.IsEmpty() )
 	{
 		URL += FString::Printf( TEXT( "&user=%s" ), *FGenericPlatformHttp::UrlEncode( UserString ) );
@@ -1570,6 +1581,8 @@ void FHttpNetworkReplayStreamer::EnumerateStreams( const FNetworkReplayVersion&
 	AddRequestToQueue( EQueuedHttpRequestType::EnumeratingSessions, HttpRequest );
 }
 
+PRAGMA_ENABLE_DEPRECATION_WARNINGS
+
 void FHttpNetworkReplayStreamer::EnumerateRecentStreams( const FNetworkReplayVersion& InReplayVersion, const int32 UserIndex, const FEnumerateStreamsCallback& Delegate )
 {
 	if ( ServerURL.IsEmpty() )
diff --git a/Engine/Source/Runtime/NetworkReplayStreaming/HttpNetworkReplayStreaming/Public/HttpNetworkReplayStreaming.h b/Engine/Source/Runtime/NetworkReplayStreaming/HttpNetworkReplayStreaming/Public/HttpNetworkReplayStreaming.h
--- a/Engine/Source/Runtime/NetworkReplayStreaming/HttpNetworkReplayStreaming/Public/HttpNetworkReplayStreaming.h
+++ b/Engine/Source/Runtime/NetworkReplayStreaming/HttpNetworkReplayStreaming/Public/HttpNetworkReplayStreaming.h
@@ -195,6 +195,8 @@ public:
 	double				LastAccessTime;
 };
 
+PRAGMA_DISABLE_DEPRECATION_WARNINGS
+
 /**
  * Http network replay streaming manager
  */
@@ -222,6 +224,8 @@ public:
 	virtual bool		IsLive() const override;
 	virtual void		DeleteFinishedStream( const FString& StreamName, const FDeleteFinishedStreamCallback& Delegate ) override;
 	virtual void		DeleteFinishedStream( const FString& StreamName, const int32 UserIndex, const FDeleteFinishedStreamCallback& Delegate ) override;
+	virtual void		EnumerateStreams( const FNetworkReplayVersion& ReplayVersion, const FString& UserString, const FString& MetaString, const FEnumerateStreamsCallback& Delegate ) override;
+	virtual void		EnumerateStreams( const FNetworkReplayVersion& InReplayVersion, const FString& UserString, const FString& MetaString, const TArray< FString >& ExtraParms, const FEnumerateStreamsCallback& Delegate ) override;
 	virtual void		EnumerateStreams( const FNetworkReplayVersion& InReplayVersion, const int32 UserIndex, const FString& MetaString, const TArray< FString >& ExtraParms, const FEnumerateStreamsCallback& Delegate ) override;
 	virtual void		EnumerateEvents( const FString& Group, const FEnumerateEventsCallback& Delegate ) override;
 	virtual void		EnumerateEvents( const FString& ReplayName, const FString& Group, const FEnumerateEventsCallback& Delegate ) override;
@@ -381,6 +385,8 @@ public:
 	int32								RefreshViewerFails;
 };
 
+PRAGMA_ENABLE_DEPRECATION_WARNINGS
+
 class HTTPNETWORKREPLAYSTREAMING_API FHttpNetworkReplayStreamingFactory : public INetworkReplayStreamingFactory, public FTickableGameObject
 {
 public:
