Flush local file streamer queue on module shutdown

diff --git a/Engine/Source/Runtime/NetworkReplayStreaming/LocalFileNetworkReplayStreaming/Private/LocalFileNetworkReplayStreaming.cpp b/Engine/Source/Runtime/NetworkReplayStreaming/LocalFileNetworkReplayStreaming/Private/LocalFileNetworkReplayStreaming.cpp
--- a/Engine/Source/Runtime/NetworkReplayStreaming/LocalFileNetworkReplayStreaming/Private/LocalFileNetworkReplayStreaming.cpp
+++ b/Engine/Source/Runtime/NetworkReplayStreaming/LocalFileNetworkReplayStreaming/Private/LocalFileNetworkReplayStreaming.cpp
@@ -3059,6 +3059,46 @@ void FLocalFileNetworkReplayStreamingFactory::Tick( float DeltaTime )
 	}
 }
 
+void FLocalFileNetworkReplayStreamingFactory::ShutdownModule()
+{
+	bool bFlushStreamersOnShutdown = true;
+	GConfig->GetBool(TEXT("LocalFileNetworkReplayStreamingFactory"), TEXT("bFlushStreamersOnShutdown"), bFlushStreamersOnShutdown, GEngineIni);
+
+	if (bFlushStreamersOnShutdown)
+	{
+		double MaxFlushTimeSeconds = -1.0;
+		GConfig->GetDouble(TEXT("LocalFileNetworkReplayStreamingFactory"), TEXT("MaxFlushTimeSeconds"), MaxFlushTimeSeconds, GEngineIni);
+
+		double BeginWaitTime = FPlatformTime::Seconds();
+		double LastTime = BeginWaitTime;
+		while (LocalFileStreamers.Num() > 0)
+		{
+			const double AppTime = FPlatformTime::Seconds();
+			const double TotalWait = AppTime - BeginWaitTime;
+
+			if ((MaxFlushTimeSeconds > 0) && (TotalWait > MaxFlushTimeSeconds))
+			{
+				UE_LOG(LogLocalFileReplay, Display, TEXT("Abandoning streamer flush after waiting %0.2f seconds"), TotalWait);
+				break;
+			}
+
+			Tick(AppTime - LastTime);
+			LastTime = AppTime;
+
+			if (LocalFileStreamers.Num() > 0)
+			{
+				FTaskGraphInterface::Get().ProcessThreadUntilIdle(ENamedThreads::GameThread);
+
+				if (FPlatformProcess::SupportsMultithreading())
+				{
+					UE_LOG(LogLocalFileReplay, Display, TEXT("Sleeping 0.1s to wait for outstanding requests."));
+					FPlatformProcess::Sleep(0.1f);
+				}
+			}
+		}
+	}
+}
+
 TStatId FLocalFileNetworkReplayStreamingFactory::GetStatId() const
 {
 	RETURN_QUICK_DECLARE_CYCLE_STAT(FLocalFileNetworkReplayStreamingFactory, STATGROUP_Tickables);
diff --git a/Engine/Source/Runtime/NetworkReplayStreaming/LocalFileNetworkReplayStreaming/Public/LocalFileNetworkReplayStreaming.h b/Engine/Source/Runtime/NetworkReplayStreaming/LocalFileNetworkReplayStreaming/Public/LocalFileNetworkReplayStreaming.h
--- a/Engine/Source/Runtime/NetworkReplayStreaming/LocalFileNetworkReplayStreaming/Public/LocalFileNetworkReplayStreaming.h
+++ b/Engine/Source/Runtime/NetworkReplayStreaming/LocalFileNetworkReplayStreaming/Public/LocalFileNetworkReplayStreaming.h
@@ -642,6 +642,8 @@ public:
 class LOCALFILENETWORKREPLAYSTREAMING_API FLocalFileNetworkReplayStreamingFactory : public INetworkReplayStreamingFactory, public FTickableGameObject
 {
 public:
+	virtual void ShutdownModule() override;
+
 	virtual TSharedPtr<INetworkReplayStreamer> CreateReplayStreamer();
 
 	/** FTickableGameObject */
