Fix non-unity compile

#rnx

diff --git a/Engine/Source/Runtime/NetworkReplayStreaming/LocalFileNetworkReplayStreaming/Private/LocalFileNetworkReplayStreaming.cpp b/Engine/Source/Runtime/NetworkReplayStreaming/LocalFileNetworkReplayStreaming/Private/LocalFileNetworkReplayStreaming.cpp
--- a/Engine/Source/Runtime/NetworkReplayStreaming/LocalFileNetworkReplayStreaming/Private/LocalFileNetworkReplayStreaming.cpp
+++ b/Engine/Source/Runtime/NetworkReplayStreaming/LocalFileNetworkReplayStreaming/Private/LocalFileNetworkReplayStreaming.cpp
@@ -13,6 +13,7 @@
 #include "Engine/World.h"
 #include "UObject/CoreOnline.h"
 #include "Serialization/LargeMemoryReader.h"
+#include "Misc/ConfigCacheIni.h"
 
 DEFINE_LOG_CATEGORY_STATIC(LogLocalFileReplay, Log, All);
 
