Adding RequestEventGroupData to replay streamer API for batched requests

diff --git a/Engine/Source/Runtime/Engine/Classes/Engine/DemoNetDriver.h b/Engine/Source/Runtime/Engine/Classes/Engine/DemoNetDriver.h
--- a/Engine/Source/Runtime/Engine/Classes/Engine/DemoNetDriver.h
+++ b/Engine/Source/Runtime/Engine/Classes/Engine/DemoNetDriver.h
@@ -826,6 +826,10 @@ public:
 	void RequestEventDataForActiveReplay(const FString& EventID, const FRequestEventDataCallback& Delegate);
 	void RequestEventDataForActiveReplay(const FString& EventID, const int32 UserIndex, const FRequestEventDataCallback& Delegate);
 
+	/** Retrieve data for all events matching the requested group, and call the passed in delegate on completion. */
+	void RequestEventGroupDataForActiveReplay(const FString& Group, const FRequestEventGroupDataCallback& Delegate);
+	void RequestEventGroupDataForActiveReplay(const FString& Group, const int32 UserIndex, const FRequestEventGroupDataCallback& Delegate);
+
 	bool IsFastForwarding() const { return bIsFastForwarding; }
 	bool IsFinalizingFastForward() const { return bIsFinalizingFastForward; }
 
diff --git a/Engine/Source/Runtime/Engine/Private/DemoNetDriver.cpp b/Engine/Source/Runtime/Engine/Private/DemoNetDriver.cpp
--- a/Engine/Source/Runtime/Engine/Private/DemoNetDriver.cpp
+++ b/Engine/Source/Runtime/Engine/Private/DemoNetDriver.cpp
@@ -2366,6 +2366,22 @@ void UDemoNetDriver::RequestEventDataForActiveReplay(const FString& EventID, con
 	}
 }
 
+void UDemoNetDriver::RequestEventGroupDataForActiveReplay(const FString& Group, const FRequestEventGroupDataCallback& Delegate)
+{
+	if (ReplayStreamer.IsValid())
+	{
+		ReplayStreamer->RequestEventGroupData(ActiveReplayName, Group, Delegate);
+	}
+}
+
+void UDemoNetDriver::RequestEventGroupDataForActiveReplay(const FString& Group, const int32 UserIndex, const FRequestEventGroupDataCallback& Delegate)
+{
+	if (ReplayStreamer.IsValid())
+	{
+		ReplayStreamer->RequestEventGroupData(ActiveReplayName, Group, UserIndex, Delegate);
+	}
+}
+
 /**
 * FReplayViewer
 * Used when demo.UseNetRelevancy enabled
diff --git a/Engine/Source/Runtime/NetworkReplayStreaming/HttpNetworkReplayStreaming/Private/HttpNetworkReplayStreaming.cpp b/Engine/Source/Runtime/NetworkReplayStreaming/HttpNetworkReplayStreaming/Private/HttpNetworkReplayStreaming.cpp
--- a/Engine/Source/Runtime/NetworkReplayStreaming/HttpNetworkReplayStreaming/Private/HttpNetworkReplayStreaming.cpp
+++ b/Engine/Source/Runtime/NetworkReplayStreaming/HttpNetworkReplayStreaming/Private/HttpNetworkReplayStreaming.cpp
@@ -944,6 +944,24 @@ void FHttpNetworkReplayStreamer::HttpRequestEventDataFinished(FHttpRequestPtr Ht
 	RequestEventDataCompleteDelegate.ExecuteIfBound(Result);
 }
 
+void FHttpNetworkReplayStreamer::RequestEventGroupData(const FString& Group, const FRequestEventGroupDataCallback& Delegate)
+{
+	UE_LOG(LogHttpReplay, Log, TEXT("FHttpNetworkReplayStreamer::RequestEventGroupData is currently unsupported."));
+	FRequestEventGroupDataResult Result;
+	Result.Result = EStreamingOperationResult::Unsupported;
+	Delegate.Execute(Result);
+}
+
+void FHttpNetworkReplayStreamer::RequestEventGroupData(const FString& ReplayName, const FString& Group, const FRequestEventGroupDataCallback& Delegate)
+{
+	RequestEventGroupData(Group, Delegate);
+}
+
+void FHttpNetworkReplayStreamer::RequestEventGroupData(const FString& ReplayName, const FString& Group, const int32 UserIndex, const FRequestEventGroupDataCallback& Delegate)
+{
+	RequestEventGroupData(Group, Delegate);
+}
+
 void FHttpNetworkReplayStreamer::GotoTimeInMS( const uint32 TimeInMS, const FGotoCallback& Delegate, EReplayCheckpointType CheckpointType )
 {
 	if ( !IsCheckpointTypeSupported(CheckpointType) )
diff --git a/Engine/Source/Runtime/NetworkReplayStreaming/HttpNetworkReplayStreaming/Public/HttpNetworkReplayStreaming.h b/Engine/Source/Runtime/NetworkReplayStreaming/HttpNetworkReplayStreaming/Public/HttpNetworkReplayStreaming.h
--- a/Engine/Source/Runtime/NetworkReplayStreaming/HttpNetworkReplayStreaming/Public/HttpNetworkReplayStreaming.h
+++ b/Engine/Source/Runtime/NetworkReplayStreaming/HttpNetworkReplayStreaming/Public/HttpNetworkReplayStreaming.h
@@ -236,6 +236,9 @@ public:
 	virtual void		RequestEventData(const FString& EventId, const FRequestEventDataCallback& Delegate) override;
 	virtual void		RequestEventData(const FString& ReplayName, const FString& EventId, const FRequestEventDataCallback& Delegate) override;
 	virtual void		RequestEventData(const FString& ReplayName, const FString& EventId, const int32 UserIndex, const FRequestEventDataCallback& Delegate) override;
+	virtual void		RequestEventGroupData(const FString& Group, const FRequestEventGroupDataCallback& Delegate) override;
+	virtual void		RequestEventGroupData(const FString& ReplayName, const FString& Group, const FRequestEventGroupDataCallback& Delegate) override;
+	virtual void		RequestEventGroupData(const FString& ReplayName, const FString& Group, const int32 UserIndex, const FRequestEventGroupDataCallback& Delegate) override;
 	virtual void		SearchEvents(const FString& EventGroup, const FSearchEventsCallback& Delegate) override;
 	virtual void		KeepReplay(const FString& ReplayName, const bool bKeep, const FKeepReplayCallback& Delegate) override;
 	virtual void		KeepReplay(const FString& ReplayName, const bool bKeep, const int32 UserIndex, const FKeepReplayCallback& Delegate) override;
diff --git a/Engine/Source/Runtime/NetworkReplayStreaming/InMemoryNetworkReplayStreaming/Private/InMemoryNetworkReplayStreaming.cpp b/Engine/Source/Runtime/NetworkReplayStreaming/InMemoryNetworkReplayStreaming/Private/InMemoryNetworkReplayStreaming.cpp
--- a/Engine/Source/Runtime/NetworkReplayStreaming/InMemoryNetworkReplayStreaming/Private/InMemoryNetworkReplayStreaming.cpp
+++ b/Engine/Source/Runtime/NetworkReplayStreaming/InMemoryNetworkReplayStreaming/Private/InMemoryNetworkReplayStreaming.cpp
@@ -285,6 +285,24 @@ void FInMemoryNetworkReplayStreamer::RequestEventData(const FString& ReplayName,
 	RequestEventData(EventID, Delegate);
 }
 
+void FInMemoryNetworkReplayStreamer::RequestEventGroupData(const FString& Group, const FRequestEventGroupDataCallback& Delegate)
+{
+	UE_LOG(LogMemoryReplay, Log, TEXT("FInMemoryNetworkReplayStreamer::RequestEventGroupData is currently unsupported."));
+	FRequestEventGroupDataResult Result;
+	Result.Result = EStreamingOperationResult::Unsupported;
+	Delegate.Execute(Result);
+}
+
+void FInMemoryNetworkReplayStreamer::RequestEventGroupData(const FString& ReplayName, const FString& Group, const FRequestEventGroupDataCallback& Delegate)
+{
+	RequestEventGroupData(Group, Delegate);
+}
+
+void FInMemoryNetworkReplayStreamer::RequestEventGroupData(const FString& ReplayName, const FString& Group, const int32 UserIndex, const FRequestEventGroupDataCallback& Delegate)
+{
+	RequestEventGroupData(Group, Delegate);
+}
+
 void FInMemoryNetworkReplayStreamer::SearchEvents(const FString& EventGroup, const FSearchEventsCallback& Delegate)
 {
 	UE_LOG(LogMemoryReplay, Log, TEXT("FInMemoryNetworkReplayStreamer::SearchEvents is currently unsupported."));
diff --git a/Engine/Source/Runtime/NetworkReplayStreaming/InMemoryNetworkReplayStreaming/Public/InMemoryNetworkReplayStreaming.h b/Engine/Source/Runtime/NetworkReplayStreaming/InMemoryNetworkReplayStreaming/Public/InMemoryNetworkReplayStreaming.h
--- a/Engine/Source/Runtime/NetworkReplayStreaming/InMemoryNetworkReplayStreaming/Public/InMemoryNetworkReplayStreaming.h
+++ b/Engine/Source/Runtime/NetworkReplayStreaming/InMemoryNetworkReplayStreaming/Public/InMemoryNetworkReplayStreaming.h
@@ -157,8 +157,10 @@ public:
 	virtual void RequestEventData(const FString& EventID, const FRequestEventDataCallback& Delegate) override;
 	virtual void RequestEventData(const FString& ReplayName, const FString& EventID, const FRequestEventDataCallback& Delegate) override;
 	virtual void RequestEventData(const FString& ReplayName, const FString& EventId, const int32 UserIndex, const FRequestEventDataCallback& Delegate) override;
+	virtual void RequestEventGroupData(const FString& Group, const FRequestEventGroupDataCallback& Delegate) override;
+	virtual void RequestEventGroupData(const FString& ReplayName, const FString& Group, const FRequestEventGroupDataCallback& Delegate) override;
+	virtual void RequestEventGroupData(const FString& ReplayName, const FString& Group, const int32 UserIndex, const FRequestEventGroupDataCallback& Delegate) override;
 	virtual void SearchEvents(const FString& EventGroup, const FSearchEventsCallback& Delegate) override;
-
 	virtual void KeepReplay(const FString& ReplayName, const bool bKeep, const int32 UserIndex, const FKeepReplayCallback& Delegate) override;
 	virtual void KeepReplay(const FString& ReplayName, const bool bKeep, const FKeepReplayCallback& Delegate) override;
 	virtual void RenameReplayFriendlyName(const FString& ReplayName, const FString& NewFriendlyName, const FRenameReplayCallback& Delegate) override;
diff --git a/Engine/Source/Runtime/NetworkReplayStreaming/LocalFileNetworkReplayStreaming/Private/LocalFileNetworkReplayStreaming.cpp b/Engine/Source/Runtime/NetworkReplayStreaming/LocalFileNetworkReplayStreaming/Private/LocalFileNetworkReplayStreaming.cpp
--- a/Engine/Source/Runtime/NetworkReplayStreaming/LocalFileNetworkReplayStreaming/Private/LocalFileNetworkReplayStreaming.cpp
+++ b/Engine/Source/Runtime/NetworkReplayStreaming/LocalFileNetworkReplayStreaming/Private/LocalFileNetworkReplayStreaming.cpp
@@ -1298,6 +1298,67 @@ void FLocalFileNetworkReplayStreamer::RequestEventData_Internal(const FString& R
 		});
 }
 
+void FLocalFileNetworkReplayStreamer::RequestEventGroupData(const FString& Group, const FRequestEventGroupDataCallback& Delegate)
+{
+	RequestEventGroupData(CurrentStreamName, Group, Delegate);
+}
+
+void FLocalFileNetworkReplayStreamer::RequestEventGroupData(const FString& ReplayName, const FString& Group, const FRequestEventGroupDataCallback& Delegate)
+{
+	RequestEventGroupData(ReplayName, Group, INDEX_NONE, Delegate);
+}
+
+void FLocalFileNetworkReplayStreamer::RequestEventGroupData(const FString& ReplayName, const FString& Group, const int32 UserIndex, const FRequestEventGroupDataCallback& Delegate)
+{
+	AddDelegateFileRequestToQueue<FRequestEventGroupDataCallback, FRequestEventGroupDataResult>(EQueuedLocalFileRequestType::RequestingEvent, Delegate,
+		[this, ReplayName, Group](TLocalFileRequestCommonData<FRequestEventGroupDataResult>& RequestData)
+		{
+			SCOPE_CYCLE_COUNTER(STAT_LocalReplay_ReadEvent);
+
+			const FString FullDemoFilename = GetDemoFullFilename(ReplayName);
+			if (!FPaths::FileExists(FullDemoFilename))
+			{
+				RequestData.DelegateResult.Result = EStreamingOperationResult::ReplayNotFound;
+			}
+			else
+			{
+				// Read stored info for this replay
+				FLocalFileReplayInfo StoredReplayInfo;
+				if (ReadReplayInfo(ReplayName, StoredReplayInfo))
+				{
+					TSharedPtr<FArchive> LocalFileAr = CreateLocalFileReader(FullDemoFilename);
+					if (LocalFileAr.IsValid())
+					{
+						for (const FLocalFileEventInfo& EventInfo : StoredReplayInfo.Events)
+						{
+							if (EventInfo.Group == Group)
+							{
+								LocalFileAr->Seek(EventInfo.EventDataOffset);
+
+								RequestData.DelegateResult.Result = EStreamingOperationResult::Success;
+								
+								FReplayEventListItem& EventItem = RequestData.DelegateResult.ReplayEventListItems.AddDefaulted_GetRef();
+								EventItem.ID = EventInfo.Id;
+								EventItem.Group = EventInfo.Group;
+								EventItem.Metadata = EventInfo.Metadata;
+								EventItem.Time1 = EventInfo.Time1;
+								EventItem.Time2 = EventInfo.Time2;
+
+								FRequestEventDataResult& EventDataResult = RequestData.DelegateResult.ReplayEventListResults.AddDefaulted_GetRef();
+								EventDataResult.Result = EStreamingOperationResult::Success;
+								EventDataResult.ReplayEventListItem.AddUninitialized(EventInfo.SizeInBytes);
+							
+								LocalFileAr->Serialize(EventDataResult.ReplayEventListItem.GetData(), EventDataResult.ReplayEventListItem.Num());
+							}
+						}
+
+						LocalFileAr = nullptr;
+					}
+				}
+			}
+		});
+}
+
 void FLocalFileNetworkReplayStreamer::SearchEvents(const FString& EventGroup, const FSearchEventsCallback& Delegate)
 {
 	UE_LOG(LogLocalFileReplay, Log, TEXT("FLocalFileNetworkReplayStreamer::SearchEvents is currently unsupported."));
diff --git a/Engine/Source/Runtime/NetworkReplayStreaming/LocalFileNetworkReplayStreaming/Public/LocalFileNetworkReplayStreaming.h b/Engine/Source/Runtime/NetworkReplayStreaming/LocalFileNetworkReplayStreaming/Public/LocalFileNetworkReplayStreaming.h
--- a/Engine/Source/Runtime/NetworkReplayStreaming/LocalFileNetworkReplayStreaming/Public/LocalFileNetworkReplayStreaming.h
+++ b/Engine/Source/Runtime/NetworkReplayStreaming/LocalFileNetworkReplayStreaming/Public/LocalFileNetworkReplayStreaming.h
@@ -432,6 +432,9 @@ public:
 	virtual void RequestEventData(const FString& EventID, const FRequestEventDataCallback& Delegate) override;
 	virtual void RequestEventData(const FString& ReplayName, const FString& EventID, const FRequestEventDataCallback& Delegate) override;
 	virtual void RequestEventData(const FString& ReplayName, const FString& EventId, const int32 UserIndex, const FRequestEventDataCallback& Delegate) override;
+	virtual void RequestEventGroupData(const FString& Group, const FRequestEventGroupDataCallback& Delegate) override;
+	virtual void RequestEventGroupData(const FString& ReplayName, const FString& Group, const FRequestEventGroupDataCallback& Delegate) override;
+	virtual void RequestEventGroupData(const FString& ReplayName, const FString& Group, const int32 UserIndex, const FRequestEventGroupDataCallback& Delegate) override;
 	virtual void SearchEvents(const FString& EventGroup, const FSearchEventsCallback& Delegate) override;
 	virtual void KeepReplay(const FString& ReplayName, const bool bKeep, const FKeepReplayCallback& Delegate) override;
 	virtual void KeepReplay(const FString& ReplayName, const bool bKeep, const int32 UserIndex, const FKeepReplayCallback& Delegate) override;
diff --git a/Engine/Source/Runtime/NetworkReplayStreaming/NetworkReplayStreaming/Public/NetworkReplayStreaming.h b/Engine/Source/Runtime/NetworkReplayStreaming/NetworkReplayStreaming/Public/NetworkReplayStreaming.h
--- a/Engine/Source/Runtime/NetworkReplayStreaming/NetworkReplayStreaming/Public/NetworkReplayStreaming.h
+++ b/Engine/Source/Runtime/NetworkReplayStreaming/NetworkReplayStreaming/Public/NetworkReplayStreaming.h
@@ -372,6 +372,16 @@ static FORCEINLINE FRequestEventDataCallback UpgradeRequestEventDelegate(const F
 }
 /** End RequestEventData Types */
 
+/** Start RequestEventGroupData Types */
+struct FRequestEventGroupDataResult : public FStreamingResultBase
+{
+	TArray<FReplayEventListItem> ReplayEventListItems;
+	TArray<FRequestEventDataResult> ReplayEventListResults;
+};
+
+DECLARE_DELEGATE_OneParam(FRequestEventGroupDataCallback, const FRequestEventGroupDataResult&);
+/** End RequestEventGroupData Types */
+
 /** Start DownloadHeader Types */
 struct FDownloadHeaderResult : public FStreamingResultBase
 {
@@ -504,6 +514,10 @@ public:
 	virtual void RequestEventData(const FString& ReplayName, const FString& EventID, const FRequestEventDataCallback& Delegate) = 0;
 	virtual void RequestEventData(const FString& ReplayName, const FString& EventId, const int32 UserIndex, const FRequestEventDataCallback& Delegate) = 0;
 
+	virtual void RequestEventGroupData(const FString& Group, const FRequestEventGroupDataCallback& Delegate) = 0;
+	virtual void RequestEventGroupData(const FString& ReplayName, const FString& Group, const FRequestEventGroupDataCallback& Delegate) = 0;
+	virtual void RequestEventGroupData(const FString& ReplayName, const FString& Group, const int32 UserIndex, const FRequestEventGroupDataCallback& Delegate) = 0;
+
 	virtual void SearchEvents(const FString& EventGroup, const FSearchEventsCallback& Delegate) = 0;
 	virtual void RefreshHeader() = 0;
 
diff --git a/Engine/Source/Runtime/NetworkReplayStreaming/NullNetworkReplayStreaming/Private/NullNetworkReplayStreaming.cpp b/Engine/Source/Runtime/NetworkReplayStreaming/NullNetworkReplayStreaming/Private/NullNetworkReplayStreaming.cpp
--- a/Engine/Source/Runtime/NetworkReplayStreaming/NullNetworkReplayStreaming/Private/NullNetworkReplayStreaming.cpp
+++ b/Engine/Source/Runtime/NetworkReplayStreaming/NullNetworkReplayStreaming/Private/NullNetworkReplayStreaming.cpp
@@ -440,6 +440,24 @@ void FNullNetworkReplayStreamer::RequestEventData(const FString& EventID, const
 	Delegate.Execute(Result);
 }
 
+void FNullNetworkReplayStreamer::RequestEventGroupData(const FString& Group, const FRequestEventGroupDataCallback& Delegate)
+{
+	UE_LOG(LogNullReplay, Log, TEXT("FNullNetworkReplayStreamer::RequestEventGroupData is currently unsupported."));
+	FRequestEventGroupDataResult Result;
+	Result.Result = EStreamingOperationResult::Unsupported;
+	Delegate.Execute(Result);
+}
+
+void FNullNetworkReplayStreamer::RequestEventGroupData(const FString& ReplayName, const FString& Group, const FRequestEventGroupDataCallback& Delegate)
+{
+	RequestEventGroupData(Group, Delegate);
+}
+
+void FNullNetworkReplayStreamer::RequestEventGroupData(const FString& ReplayName, const FString& Group, const int32 UserIndex, const FRequestEventGroupDataCallback& Delegate)
+{
+	RequestEventGroupData(Group, Delegate);
+}
+
 void FNullNetworkReplayStreamer::SearchEvents(const FString& EventGroup, const FSearchEventsCallback& Delegate)
 {
 	UE_LOG(LogNullReplay, Log, TEXT("FNullNetworkReplayStreamer::SearchEvents is currently unsupported."));
diff --git a/Engine/Source/Runtime/NetworkReplayStreaming/NullNetworkReplayStreaming/Public/NullNetworkReplayStreaming.h b/Engine/Source/Runtime/NetworkReplayStreaming/NullNetworkReplayStreaming/Public/NullNetworkReplayStreaming.h
--- a/Engine/Source/Runtime/NetworkReplayStreaming/NullNetworkReplayStreaming/Public/NullNetworkReplayStreaming.h
+++ b/Engine/Source/Runtime/NetworkReplayStreaming/NullNetworkReplayStreaming/Public/NullNetworkReplayStreaming.h
@@ -80,6 +80,9 @@ public:
 	virtual void RequestEventData(const FString& EventID, const FRequestEventDataCallback& Delegate) override;
 	virtual void RequestEventData(const FString& ReplayName, const FString& EventID, const FRequestEventDataCallback& Delegate) override;
 	virtual void RequestEventData(const FString& ReplayName, const FString& EventId, const int32 UserIndex, const FRequestEventDataCallback& Delegate) override;
+	virtual void RequestEventGroupData(const FString& Group, const FRequestEventGroupDataCallback& Delegate) override;
+	virtual void RequestEventGroupData(const FString& ReplayName, const FString& Group, const FRequestEventGroupDataCallback& Delegate) override;
+	virtual void RequestEventGroupData(const FString& ReplayName, const FString& Group, const int32 UserIndex, const FRequestEventGroupDataCallback& Delegate) override;
 	virtual void SearchEvents(const FString& EventGroup, const FSearchEventsCallback& Delegate) override;
 	virtual void KeepReplay(const FString& ReplayName, const bool bKeep, const FKeepReplayCallback& Delegate) override;
 	virtual void KeepReplay(const FString& ReplayName, const bool bKeep, const int32 UserIndex, const FKeepReplayCallback& Delegate) override;
diff --git a/Engine/Source/Runtime/NetworkReplayStreaming/SaveGameNetworkReplayStreaming/Private/SaveGameNetworkReplayStreaming.cpp b/Engine/Source/Runtime/NetworkReplayStreaming/SaveGameNetworkReplayStreaming/Private/SaveGameNetworkReplayStreaming.cpp
--- a/Engine/Source/Runtime/NetworkReplayStreaming/SaveGameNetworkReplayStreaming/Private/SaveGameNetworkReplayStreaming.cpp
+++ b/Engine/Source/Runtime/NetworkReplayStreaming/SaveGameNetworkReplayStreaming/Private/SaveGameNetworkReplayStreaming.cpp
@@ -1476,6 +1476,101 @@ void FSaveGameNetworkReplayStreamer::RequestEventData_Internal(const FString& Re
 	}
 }
 
+void FSaveGameNetworkReplayStreamer::RequestEventGroupData(const FString& Group, const FRequestEventGroupDataCallback& Delegate)
+{
+	UE_LOG(LogSaveGameReplay, Warning, TEXT("FSaveGameNetworkReplayStreamer::RequestEventGroupData: No replay name available, defaulting to FLocalFileStreamFArchive::RequestEventData"));
+	FLocalFileNetworkReplayStreamer::RequestEventGroupData(Group, Delegate);
+}
+
+void FSaveGameNetworkReplayStreamer::RequestEventGroupData(const FString& ReplayName, const FString& Group, const FRequestEventGroupDataCallback& Delegate)
+{
+	if (!IsSaveGameFileName(ReplayName))
+	{
+		FLocalFileNetworkReplayStreamer::RequestEventGroupData(ReplayName, Group, INDEX_NONE, Delegate);
+	}
+	else
+	{
+		UE_LOG(LogSaveGameReplay, Warning, TEXT("FSaveGameNetworkReplayStreamer::RequestEventGroupData - Invalid UserIndex."));
+		Delegate.ExecuteIfBound(FRequestEventGroupDataResult());
+	}
+}
+
+void FSaveGameNetworkReplayStreamer::RequestEventGroupData(const FString& ReplayName, const FString& Group, const int32 UserIndex, const FRequestEventGroupDataCallback& Delegate)
+{
+	if (!IsSaveGameFileName(ReplayName))
+	{
+		FLocalFileNetworkReplayStreamer::RequestEventGroupData(ReplayName, Group, UserIndex, Delegate);
+	}
+	else if (UserIndex != INDEX_NONE)
+	{
+		RequestEventGroupDataSaved(ReplayName, Group, UserIndex, Delegate);
+	}
+	else
+	{
+		UE_LOG(LogSaveGameReplay, Warning, TEXT("FSaveGameNetworkReplayStreamer::RequestEventGroupData - Invalid UserIndex."));
+		Delegate.ExecuteIfBound(FRequestEventGroupDataResult());
+	}
+}
+
+void FSaveGameNetworkReplayStreamer::RequestEventGroupDataSaved(const FString& ReplayName, const FString& Group, const int32 UserIndex, const FRequestEventGroupDataCallback& Delegate)
+{
+	using TAsyncTypes = SaveGameReplay::TAsyncTypes<FRequestEventGroupDataResult>;
+
+	auto AsyncWork = [this, ReplayName, Group, UserIndex]()
+	{
+		auto SharedResult = TAsyncTypes::MakeSharedResult();
+		RequestEventGroupData_Internal(ReplayName, Group, UserIndex, *SharedResult.Get());
+		return SharedResult;
+	};
+
+	auto PostAsyncWork = [Delegate](const FRequestEventGroupDataResult& Result)
+	{
+		Delegate.ExecuteIfBound(Result);
+	};
+
+	SaveGameReplay::FAsyncTaskManager::Get().StartTask<FRequestEventGroupDataResult>(*this, TEXT("RequestEventGroupData"), AsyncWork, PostAsyncWork);
+}
+
+void FSaveGameNetworkReplayStreamer::RequestEventGroupData_Internal(const FString& ReplayName, const FString& Group, const int32 UserIndex, FRequestEventGroupDataResult& Result)
+{
+	ISaveGameSystem* SaveGameSystem = IPlatformFeaturesModule::Get().GetSaveGameSystem();
+	if (!ensureMsgf(SaveGameSystem, TEXT("FSaveGameNetworkReplayStreamer::RequestEventGroupDataSaved: Unable to retrieve save game systems")))
+	{
+		return;
+	}
+
+	FSaveGameSanitizedNames SanitizedNames;
+	if (!StreamNameToSanitizedNames(ReplayName, SanitizedNames))
+	{
+		return;
+	}
+
+	FSaveGameMetaData MetaData;
+	if (ReadMetaDataFromSaveGame(*SaveGameSystem, SanitizedNames, UserIndex, MetaData, Result))
+	{
+		for (int32 i = 0; i < MetaData.VersionedInfo.Events.ReplayEvents.Num(); ++i)
+		{
+			const FReplayEventListItem& EventInfo = MetaData.VersionedInfo.Events.ReplayEvents[i];
+
+			if (EventInfo.Group == Group)
+			{
+				Result.Result = EStreamingOperationResult::Success;
+				
+				FReplayEventListItem& EventItem = Result.ReplayEventListItems.AddDefaulted_GetRef();
+				EventItem.ID = EventInfo.ID;
+				EventItem.Group = EventInfo.Group;
+				EventItem.Metadata = EventInfo.Metadata;
+				EventItem.Time1 = EventInfo.Time1;
+				EventItem.Time2 = EventInfo.Time2;
+
+				FRequestEventDataResult& EventDataResult = Result.ReplayEventListResults.AddDefaulted_GetRef();
+				EventDataResult.Result = EStreamingOperationResult::Success;
+				EventDataResult.ReplayEventListItem = MoveTemp(MetaData.VersionedInfo.EventData[i]);
+			}
+		}
+	}
+}
+
 bool FSaveGameNetworkReplayStreamer::ReadMetaDataFromLocalStream(FArchive& StreamArchive, FSaveGameMetaData& OutMetaData) const
 {
 	check(StreamArchive.IsLoading());
diff --git a/Engine/Source/Runtime/NetworkReplayStreaming/SaveGameNetworkReplayStreaming/Public/SaveGameNetworkReplayStreaming.h b/Engine/Source/Runtime/NetworkReplayStreaming/SaveGameNetworkReplayStreaming/Public/SaveGameNetworkReplayStreaming.h
--- a/Engine/Source/Runtime/NetworkReplayStreaming/SaveGameNetworkReplayStreaming/Public/SaveGameNetworkReplayStreaming.h
+++ b/Engine/Source/Runtime/NetworkReplayStreaming/SaveGameNetworkReplayStreaming/Public/SaveGameNetworkReplayStreaming.h
@@ -58,7 +58,9 @@ public:
 	virtual void RequestEventData(const FString& EventID, const FRequestEventDataCallback& Delegate) override;
 	virtual void RequestEventData(const FString& ReplayName, const FString& EventID, const FRequestEventDataCallback& Delegate) override;
 	virtual void RequestEventData(const FString& ReplayName, const FString& EventId, const int32 UserIndex, const FRequestEventDataCallback& Delegate) override;
-
+	virtual void RequestEventGroupData(const FString& Group, const FRequestEventGroupDataCallback& Delegate) override;
+	virtual void RequestEventGroupData(const FString& ReplayName, const FString& Group, const FRequestEventGroupDataCallback& Delegate) override;
+	virtual void RequestEventGroupData(const FString& ReplayName, const FString& Group, const int32 UserIndex, const FRequestEventGroupDataCallback& Delegate) override;
 protected:
 
 	void StartStreamingSaved(const FString& CustomName, const FString& FriendlyName, const TArray<int32>& UserIndices, bool bRecord, const FNetworkReplayVersion& ReplayVersion, const FStartStreamingCallback& Delegate);
@@ -69,6 +71,7 @@ protected:
 	void RenameReplaySaved(const FString& ReplayName, const FString& NewName, const int32 UserIndex, const FRenameReplayCallback& Delegate);
 	void EnumerateEventsSaved(const FString& ReplayName, const FString& Group, const int32 UserIndex, const FEnumerateEventsCallback& Delegate) const;
 	void RequestEventDataSaved(const FString& ReplayName, const FString& EventID, const int32 UserIndex, const FRequestEventDataCallback& Delegate);
+	void RequestEventGroupDataSaved(const FString& ReplayName, const FString& Group, const int32 UserIndex, const FRequestEventGroupDataCallback& Delegate);
 
 	void StartStreaming_Internal(const FString& CustomName, const FString& FriendlyName, const TArray<int32>& UserIndices, bool bRecord, const FNetworkReplayVersion& ReplayVersion, FStartStreamingResult& Result);
 	void DeleteFinishedStream_Internal(const FString& ReplayName, const int32 UserIndex, FDeleteFinishedStreamResult& Result) const;
@@ -77,6 +80,7 @@ protected:
 	void RenameReplayFriendlyName_Internal(const FString& ReplayName, const FString& NewFriendlyName, const int32 UserIndex, FRenameReplayResult& Result) const;
 	void EnumerateEvents_Internal(const FString& ReplayName, const FString& Group, const int32 UserIndex, FEnumerateEventsResult& Result) const;
 	void RequestEventData_Internal(const FString& ReplayName, const FString& EventID, const int32 UserIndex, FRequestEventDataResult& Result);
+	void RequestEventGroupData_Internal(const FString& ReplayName, const FString& Group, const int32 UserIndex, FRequestEventGroupDataResult& Result);
 
 	struct FSaveGameReplayVersionedInfo
 	{
