Add check to prevent save game replay cheats from running on a dedicated server

diff --git a/Engine/Source/Runtime/NetworkReplayStreaming/SaveGameNetworkReplayStreaming/Private/SaveGameNetworkReplayStreaming.cpp b/Engine/Source/Runtime/NetworkReplayStreaming/SaveGameNetworkReplayStreaming/Private/SaveGameNetworkReplayStreaming.cpp
--- a/Engine/Source/Runtime/NetworkReplayStreaming/SaveGameNetworkReplayStreaming/Private/SaveGameNetworkReplayStreaming.cpp
+++ b/Engine/Source/Runtime/NetworkReplayStreaming/SaveGameNetworkReplayStreaming/Private/SaveGameNetworkReplayStreaming.cpp
@@ -242,6 +242,12 @@ namespace SaveGameReplay
 
 		static void RunCommand(const TArray<FString>& Params, const TFunction<void(const TSharedPtr<INetworkReplayStreamer>&, const FString&, const int32)> CommandToRun)
 		{
+			if (IsRunningDedicatedServer())
+			{
+				UE_LOG(LogSaveGameReplay, Warning, TEXT("FSaveGameMoveFileHelper commands are client only."));
+				return;
+			}
+
 			const TCHAR* StreamerOverride = nullptr;
 			if (Params.Num() == 1)
 			{
