Adjusting function moves in local file streamer

diff --git a/Engine/Source/Runtime/NetworkReplayStreaming/LocalFileNetworkReplayStreaming/Public/LocalFileNetworkReplayStreaming.h b/Engine/Source/Runtime/NetworkReplayStreaming/LocalFileNetworkReplayStreaming/Public/LocalFileNetworkReplayStreaming.h
--- a/Engine/Source/Runtime/NetworkReplayStreaming/LocalFileNetworkReplayStreaming/Public/LocalFileNetworkReplayStreaming.h
+++ b/Engine/Source/Runtime/NetworkReplayStreaming/LocalFileNetworkReplayStreaming/Public/LocalFileNetworkReplayStreaming.h
@@ -213,7 +213,9 @@ namespace EQueuedLocalFileRequestType
 class FCachedFileRequest
 {
 public:
-	FCachedFileRequest( const TArray<uint8>& InRequestData, const double InLastAccessTime ) : RequestData( InRequestData ), LastAccessTime( InLastAccessTime )
+	FCachedFileRequest(const TArray<uint8>& InRequestData, const double InLastAccessTime) 
+		: RequestData(InRequestData)
+		, LastAccessTime(InLastAccessTime)
 	{
 	}
 
@@ -250,7 +252,7 @@ protected:
 class FGenericQueuedLocalFileRequest : public FQueuedLocalFileRequest, public TSharedFromThis<FGenericQueuedLocalFileRequest, ESPMode::ThreadSafe>
 {
 public:
-	FGenericQueuedLocalFileRequest(const TSharedPtr<FLocalFileNetworkReplayStreamer>& InStreamer, EQueuedLocalFileRequestType::Type InType, TFunction<void()>& InFunction, TFunction<void()>& InCompletionCallback)
+	FGenericQueuedLocalFileRequest(const TSharedPtr<FLocalFileNetworkReplayStreamer>& InStreamer, EQueuedLocalFileRequestType::Type InType, TFunction<void()>&& InFunction, TFunction<void()>&& InCompletionCallback)
 		: FQueuedLocalFileRequest(InStreamer, InType)
 		, RequestFunction(MoveTemp(InFunction))
 		, CompletionCallback(MoveTemp(InCompletionCallback))
@@ -305,7 +307,7 @@ template <typename StorageType>
 class TGenericQueuedLocalFileRequest : public FQueuedLocalFileRequest, public TSharedFromThis<TGenericQueuedLocalFileRequest<StorageType>, ESPMode::ThreadSafe>
 {
 public:
-	TGenericQueuedLocalFileRequest(const TSharedPtr<FLocalFileNetworkReplayStreamer>& InStreamer, EQueuedLocalFileRequestType::Type InType, TFunction<void(StorageType&)>& InFunction, TFunction<void(StorageType&)>& InCompletionCallback)
+	TGenericQueuedLocalFileRequest(const TSharedPtr<FLocalFileNetworkReplayStreamer>& InStreamer, EQueuedLocalFileRequestType::Type InType, TFunction<void(StorageType&)>&& InFunction, TFunction<void(StorageType&)>&& InCompletionCallback)
 		: FQueuedLocalFileRequest(InStreamer, InType)
 		, RequestFunction(MoveTemp(InFunction))
 		, CompletionCallback(MoveTemp(InCompletionCallback))
@@ -366,8 +368,8 @@ template <typename DelegateResultType>
 class TGenericCachedLocalFileRequest : public TGenericQueuedLocalFileRequest<TLocalFileRequestCommonData<DelegateResultType>>
 {
 public:
-	TGenericCachedLocalFileRequest(int32 InCacheKey, const TSharedPtr<FLocalFileNetworkReplayStreamer>& InStreamer, EQueuedLocalFileRequestType::Type InType, TFunction<void(TLocalFileRequestCommonData<DelegateResultType>&)>& InFunction, TFunction<void(TLocalFileRequestCommonData<DelegateResultType>&)>& InCompletionCallback)
-		: TGenericQueuedLocalFileRequest<TLocalFileRequestCommonData<DelegateResultType>>(InStreamer, InType, InFunction, InCompletionCallback)
+	TGenericCachedLocalFileRequest(int32 InCacheKey, const TSharedPtr<FLocalFileNetworkReplayStreamer>& InStreamer, EQueuedLocalFileRequestType::Type InType, TFunction<void(TLocalFileRequestCommonData<DelegateResultType>&)>&& InFunction, TFunction<void(TLocalFileRequestCommonData<DelegateResultType>&)>&& InCompletionCallback)
+		: TGenericQueuedLocalFileRequest<TLocalFileRequestCommonData<DelegateResultType>>(InStreamer, InType, MoveTemp(InFunction), MoveTemp(InCompletionCallback))
 		, CacheKey(InCacheKey)
 	{}
 
@@ -483,13 +485,13 @@ public:
 
 	void AddSimpleRequestToQueue(EQueuedLocalFileRequestType::Type RequestType, TFunction<void()>&& InFunction, TFunction<void()>&& InCompletionCallback)
 	{
-		QueuedRequests.Add(MakeShared<FGenericQueuedLocalFileRequest, ESPMode::ThreadSafe>(AsShared(), RequestType, InFunction, InCompletionCallback));
+		QueuedRequests.Add(MakeShared<FGenericQueuedLocalFileRequest, ESPMode::ThreadSafe>(AsShared(), RequestType, MoveTemp(InFunction), MoveTemp(InCompletionCallback)));
 	}
 
 	template <typename StorageType>
 	void AddGenericRequestToQueue(EQueuedLocalFileRequestType::Type RequestType, TFunction<void(StorageType&)>&& InFunction, TFunction<void(StorageType&)>&& InCompletionCallback)
 	{
-		QueuedRequests.Add(MakeShared<TGenericQueuedLocalFileRequest<StorageType>, ESPMode::ThreadSafe>(AsShared(), RequestType, InFunction, InCompletionCallback));
+		QueuedRequests.Add(MakeShared<TGenericQueuedLocalFileRequest<StorageType>, ESPMode::ThreadSafe>(AsShared(), RequestType, MoveTemp(InFunction), MoveTemp(InCompletionCallback)));
 	}
 
 	template<typename DelegateResultType>
@@ -511,7 +513,7 @@ public:
 	template<typename DelegateResultType>
 	void AddCachedFileRequestToQueue(EQueuedLocalFileRequestType::Type RequestType, int32 InCacheKey, TFunction<void(TLocalFileRequestCommonData<DelegateResultType>&)>&& InFunction, TFunction<void(TLocalFileRequestCommonData<DelegateResultType>&)>&& InCompletionCallback)
 	{
-		QueuedRequests.Add(MakeShared<TGenericCachedLocalFileRequest<DelegateResultType>, ESPMode::ThreadSafe>(InCacheKey, AsShared(), RequestType, InFunction, InCompletionCallback));
+		QueuedRequests.Add(MakeShared<TGenericCachedLocalFileRequest<DelegateResultType>, ESPMode::ThreadSafe>(InCacheKey, AsShared(), RequestType, MoveTemp(InFunction), MoveTemp(InCompletionCallback)));
 	}
 
 	/** Map of chunk index to cached value */
