Add ability to encrypt local file replay chunks

diff --git a/Engine/Source/Runtime/NetworkReplayStreaming/LocalFileNetworkReplayStreaming/Private/LocalFileNetworkReplayStreaming.cpp b/Engine/Source/Runtime/NetworkReplayStreaming/LocalFileNetworkReplayStreaming/Private/LocalFileNetworkReplayStreaming.cpp
--- a/Engine/Source/Runtime/NetworkReplayStreaming/LocalFileNetworkReplayStreaming/Private/LocalFileNetworkReplayStreaming.cpp
+++ b/Engine/Source/Runtime/NetworkReplayStreaming/LocalFileNetworkReplayStreaming/Private/LocalFileNetworkReplayStreaming.cpp
@@ -21,6 +21,9 @@ DECLARE_STATS_GROUP(TEXT("LocalReplay"), STATGROUP_LocalReplay, STATCAT_Advanced
 DECLARE_CYCLE_STAT(TEXT("Local replay compress time"), STAT_LocalReplay_CompressTime, STATGROUP_LocalReplay);
 DECLARE_CYCLE_STAT(TEXT("Local replay decompress time"), STAT_LocalReplay_DecompressTime, STATGROUP_LocalReplay);
 
+DECLARE_CYCLE_STAT(TEXT("Local replay encrypt time"), STAT_LocalReplay_EncryptTime, STATGROUP_LocalReplay);
+DECLARE_CYCLE_STAT(TEXT("Local replay decrypt time"), STAT_LocalReplay_DecryptTime, STATGROUP_LocalReplay);
+
 DECLARE_CYCLE_STAT(TEXT("Local replay read info"), STAT_LocalReplay_ReadReplayInfo, STATGROUP_LocalReplay);
 DECLARE_CYCLE_STAT(TEXT("Local replay write info"), STAT_LocalReplay_WriteReplayInfo, STATGROUP_LocalReplay);
 
@@ -51,6 +54,7 @@ namespace LocalFileReplay
 		HISTORY_RECORDED_TIMESTAMP				= 3,
 		HISTORY_STREAM_CHUNK_TIMES				= 4,
 		HISTORY_FRIENDLY_NAME_ENCODING			= 5,
+		HISTORY_ENCRYPTION						= 6,
 
 		// -----<new versions can be added before this line>-------------------------------------------------
 		HISTORY_PLUS_ONE,
@@ -239,6 +243,23 @@ bool FLocalFileNetworkReplayStreamer::ReadReplayInfo(FArchive& Archive, FLocalFi
 				Info.bCompressed = (Compressed != 0);
 			}
 
+			if (FileVersion >= LocalFileReplay::HISTORY_ENCRYPTION)
+			{
+				uint32 Encrypted;
+				Archive << Encrypted;
+
+				Info.bEncrypted = (Encrypted != 0);
+
+				Archive << Info.EncryptionKey;
+			}
+
+			if (!Info.bIsLive && Info.bEncrypted && (Info.EncryptionKey.Num() == 0))
+			{
+				UE_LOG(LogLocalFileReplay, Error, TEXT("ReadReplayInfo: Completed replay is marked encrypted but has no key!"));
+				Archive.SetError();
+				return false;
+			}
+
 			int64 TotalSize = Archive.TotalSize();
 
 			// now look for all chunks
@@ -300,21 +321,10 @@ bool FLocalFileNetworkReplayStreamer::ReadReplayInfo(FArchive& Archive, FLocalFi
 					
 					if ((Checkpoint.SizeInBytes < 0) || ((Checkpoint.EventDataOffset + Checkpoint.SizeInBytes) > TotalSize))
 					{
-						UE_LOG(LogLocalFileReplay, Error, TEXT("ReadReplayInfo: Invalid checkpoint size: %lld"), Checkpoint.SizeInBytes);
+						UE_LOG(LogLocalFileReplay, Error, TEXT("ReadReplayInfo: Invalid checkpoint disk size: %lld"), Checkpoint.SizeInBytes);
 						Archive.SetError();
 						return false;
 					}
-
-					if (Info.bCompressed)
-					{
-						int32 DecompressedSize = GetDecompressedSize(Archive);
-						if (DecompressedSize < 0)
-						{
-							UE_LOG(LogLocalFileReplay, Error, TEXT("ReadReplayInfo: Invalid decompressed checkpoint size: %d"), DecompressedSize);
-							Archive.SetError();
-							return false;
-						}
-					}
 				}
 				break;
 				case ELocalFileChunkType::ReplayData:
@@ -335,31 +345,43 @@ bool FLocalFileNetworkReplayStreamer::ReadReplayInfo(FArchive& Archive, FLocalFi
 						DataChunk.SizeInBytes = Chunk.SizeInBytes;
 					}
 
-					DataChunk.ReplayDataOffset = Archive.Tell();
+					if (FileVersion < LocalFileReplay::HISTORY_ENCRYPTION)
+					{
+						DataChunk.ReplayDataOffset = Archive.Tell();
+	
+						if (Info.bCompressed)
+						{
+							PRAGMA_DISABLE_DEPRECATION_WARNINGS
+							DataChunk.MemorySizeInBytes = GetDecompressedSize(Archive);
+							PRAGMA_ENABLE_DEPRECATION_WARNINGS
+						}
+						else
+						{
+							DataChunk.MemorySizeInBytes = DataChunk.SizeInBytes;
+						}
+					}
+					else
+					{
+						Archive << DataChunk.MemorySizeInBytes;
+
+						DataChunk.ReplayDataOffset = Archive.Tell();
+					}
 
 					if ((DataChunk.SizeInBytes < 0) || ((DataChunk.ReplayDataOffset + DataChunk.SizeInBytes) > TotalSize))
 					{
-						UE_LOG(LogLocalFileReplay, Error, TEXT("ReadReplayInfo: Invalid stream chunk size: %lld"), DataChunk.SizeInBytes);
+						UE_LOG(LogLocalFileReplay, Error, TEXT("ReadReplayInfo: Invalid stream chunk disk size: %lld"), DataChunk.SizeInBytes);
 						Archive.SetError();
 						return false;
 					}
 
-					if (Info.bCompressed)
-					{
-						int32 DecompressedSize = GetDecompressedSize(Archive);
-						if (DecompressedSize < 0)
-						{
-							UE_LOG(LogLocalFileReplay, Error, TEXT("ReadReplayInfo: Invalid decompressed replay data size: %d"), DecompressedSize);
-							Archive.SetError();
-							return false;
-						}
-							
-						Info.TotalDataSizeInBytes += DecompressedSize;
-					}
-					else
+					if (DataChunk.MemorySizeInBytes < 0)
 					{
-						Info.TotalDataSizeInBytes += DataChunk.SizeInBytes;
+						UE_LOG(LogLocalFileReplay, Error, TEXT("ReadReplayInfo: Invalid stream chunk memory size: %d"), DataChunk.MemorySizeInBytes);
+						Archive.SetError();
+						return false;
 					}
+
+					Info.TotalDataSizeInBytes += DataChunk.MemorySizeInBytes;
 				}
 				break;
 				case ELocalFileChunkType::Event:
@@ -380,7 +402,7 @@ bool FLocalFileNetworkReplayStreamer::ReadReplayInfo(FArchive& Archive, FLocalFi
 
 					if ((Event.SizeInBytes < 0) || ((Event.EventDataOffset + Event.SizeInBytes) > TotalSize))
 					{
-						UE_LOG(LogLocalFileReplay, Error, TEXT("ReadReplayInfo: Invalid event size: %lld"), Event.SizeInBytes);
+						UE_LOG(LogLocalFileReplay, Error, TEXT("ReadReplayInfo: Invalid event disk size: %lld"), Event.SizeInBytes);
 						Archive.SetError();
 						return false;
 					}
@@ -565,6 +587,25 @@ bool FLocalFileNetworkReplayStreamer::WriteReplayInfo(FArchive& Archive, const F
 		Archive << Compressed;
 	}
 
+	if (SerializationInfo.FileVersion >= LocalFileReplay::HISTORY_ENCRYPTION)
+	{
+		uint32 Encrypted = SupportsEncryption() ? 1 : 0;
+		Archive << Encrypted;
+
+		TArray<uint8> KeyToWrite;
+
+		if (InReplayInfo.bIsLive)
+		{
+			KeyToWrite.AddZeroed(InReplayInfo.EncryptionKey.Num());
+		}
+		else
+		{
+			KeyToWrite = InReplayInfo.EncryptionKey;
+		}
+
+		Archive << KeyToWrite;
+	}
+
 	return !Archive.IsError();
 }
 
@@ -728,6 +769,11 @@ void FLocalFileNetworkReplayStreamer::StartStreaming(const FStartStreamingParame
 				RequestData.ReplayInfo.bIsLive = true;
 				RequestData.ReplayInfo.Timestamp = FDateTime::Now();
 
+				if (SupportsEncryption())
+				{
+					GenerateEncryptionKey(RequestData.ReplayInfo.EncryptionKey);
+				}
+
 				WriteReplayInfo(CurrentStreamName, RequestData.ReplayInfo);
 
 				RequestData.DelegateResult.Result = EStreamingOperationResult::Success;
@@ -807,8 +853,10 @@ void FLocalFileNetworkReplayStreamer::StopStreaming()
 			{
 				if (ReadReplayInfo(CurrentStreamName, ReplayInfo))
 				{
+					// Set the final values of these header properties
 					ReplayInfo.bIsLive = false;
 					ReplayInfo.LengthInMS = TotalLengthInMS;
+					ReplayInfo.EncryptionKey = CurrentReplayInfo.EncryptionKey;
 
 					WriteReplayInfo(CurrentStreamName, ReplayInfo);
 				}
@@ -1040,7 +1088,7 @@ void FLocalFileNetworkReplayStreamer::AddOrUpdateEvent(const FString& Name, cons
 	EventName = CurrentStreamName + TEXT("_") + EventName;
 
 	AddGenericRequestToQueue<FLocalFileReplayInfo>(EQueuedLocalFileRequestType::UpdatingEvent,
-		[this, EventName, Group, TimeInMS, Meta, Data](FLocalFileReplayInfo& ReplayInfo)
+		[this, EventName, Group, TimeInMS, Meta, Data, EncryptionKey=CurrentReplayInfo.EncryptionKey](FLocalFileReplayInfo& ReplayInfo)
 		{
 			SCOPE_CYCLE_COUNTER(STAT_LocalReplay_FlushEvent);
 
@@ -1061,6 +1109,24 @@ void FLocalFileNetworkReplayStreamer::AddOrUpdateEvent(const FString& Name, cons
 						}
 					}
 
+					TArray<uint8> EncryptedData;
+
+					if (SupportsEncryption())
+					{
+						SCOPE_CYCLE_COUNTER(STAT_LocalReplay_EncryptTime);
+
+						if (!EncryptBuffer(Data, EncryptedData, EncryptionKey))
+						{
+							UE_LOG(LogLocalFileReplay, Warning, TEXT("FLocalFileNetworkReplayStreamer::AddOrUpdateEvent - EncryptBuffer failed"));
+							ReplayInfo.bIsValid = false;
+							return;
+						}
+					}
+					else
+					{
+						EncryptedData = Data;
+					}
+
 					// serialize event to temporary location
 					FArrayWriter Writer;
 
@@ -1089,10 +1155,10 @@ void FLocalFileNetworkReplayStreamer::AddOrUpdateEvent(const FString& Name, cons
 					uint32 Time2 = TimeInMS;
 					Writer << Time2;
 
-					int32 EventSize = Data.Num();
+					int32 EventSize = EncryptedData.Num();
 					Writer << EventSize;
 
-					Writer.Serialize((void*)Data.GetData(), Data.Num());
+					Writer.Serialize((void*)EncryptedData.GetData(), EncryptedData.Num());
 
 					int32 ChunkSize = Writer.Tell() - MetadataPos;
 
@@ -1134,9 +1200,7 @@ void FLocalFileNetworkReplayStreamer::AddOrUpdateEvent(const FString& Name, cons
 		{
 			if (ReplayInfo.bIsValid)
 			{
-				const int32 TotalLengthInMS = CurrentReplayInfo.LengthInMS;
-				CurrentReplayInfo = ReplayInfo;
-				CurrentReplayInfo.LengthInMS = TotalLengthInMS;
+				UpdateCurrentReplayInfo(ReplayInfo);
 			}
 		});
 }
@@ -1244,10 +1308,38 @@ void FLocalFileNetworkReplayStreamer::RequestEventData_Internal(const FString& R
 							{
 								LocalFileAr->Seek(EventInfo.EventDataOffset);
 
-								RequestData.DelegateResult.Result = EStreamingOperationResult::Success;
-								RequestData.DelegateResult.ReplayEventListItem.AddUninitialized(EventInfo.SizeInBytes);
+								TArray<uint8> EventData;
+
+								EventData.AddUninitialized(EventInfo.SizeInBytes);
+								LocalFileAr->Serialize(EventData.GetData(), EventData.Num());
 
-								LocalFileAr->Serialize(RequestData.DelegateResult.ReplayEventListItem.GetData(), RequestData.DelegateResult.ReplayEventListItem.Num());
+								if (StoredReplayInfo.bEncrypted)
+								{
+									if (SupportsEncryption())
+									{
+										SCOPE_CYCLE_COUNTER(STAT_LocalReplay_DecryptTime);
+
+										TArray<uint8> PlaintextData;
+										
+										if (!DecryptBuffer(EventData, PlaintextData, StoredReplayInfo.EncryptionKey))
+										{
+											UE_LOG(LogLocalFileReplay, Error, TEXT("FLocalFileNetworkReplayStreamer::RequestEventData_Internal. DecryptBuffer failed."));
+											RequestData.DelegateResult.Result = EStreamingOperationResult::Unspecified;
+											break;
+										}
+
+										EventData = MoveTemp(PlaintextData);
+									}
+									else
+									{
+										UE_LOG(LogLocalFileReplay, Error, TEXT("FLocalFileNetworkReplayStreamer::RequestEventData_Internal. Encrypted event but streamer does not support encryption."));
+										RequestData.DelegateResult.Result = EStreamingOperationResult::Unspecified;
+										break;
+									}
+								}
+
+								RequestData.DelegateResult.Result = EStreamingOperationResult::Success;
+								RequestData.DelegateResult.ReplayEventListItem = MoveTemp(EventData);
 
 								break;
 							}
@@ -1306,11 +1398,39 @@ void FLocalFileNetworkReplayStreamer::RequestEventGroupData(const FString& Repla
 								EventItem.Time1 = EventInfo.Time1;
 								EventItem.Time2 = EventInfo.Time2;
 
+								TArray<uint8> EventData;
+
+								EventData.AddUninitialized(EventInfo.SizeInBytes);
+								LocalFileAr->Serialize(EventData.GetData(), EventData.Num());
+
+								if (StoredReplayInfo.bEncrypted)
+								{
+									if (SupportsEncryption())
+									{
+										SCOPE_CYCLE_COUNTER(STAT_LocalReplay_DecryptTime);
+
+										TArray<uint8> PlaintextData;
+										
+										if (!DecryptBuffer(EventData, PlaintextData, StoredReplayInfo.EncryptionKey))
+										{
+											UE_LOG(LogLocalFileReplay, Error, TEXT("FLocalFileNetworkReplayStreamer::RequestEventData_Internal. DecryptBuffer failed."));
+											RequestData.DelegateResult.Result = EStreamingOperationResult::Unspecified;
+											break;
+										}
+
+										EventData = MoveTemp(PlaintextData);
+									}
+									else
+									{
+										UE_LOG(LogLocalFileReplay, Error, TEXT("FLocalFileNetworkReplayStreamer::RequestEventData_Internal. Encrypted event but streamer does not support encryption."));
+										RequestData.DelegateResult.Result = EStreamingOperationResult::Unspecified;
+										break;
+									}
+								}
+
 								FRequestEventDataResult& EventDataResult = RequestData.DelegateResult.ReplayEventListResults.AddDefaulted_GetRef();
 								EventDataResult.Result = EStreamingOperationResult::Success;
-								EventDataResult.ReplayEventListItem.AddUninitialized(EventInfo.SizeInBytes);
-							
-								LocalFileAr->Serialize(EventDataResult.ReplayEventListItem.GetData(), EventDataResult.ReplayEventListItem.Num());
+								EventDataResult.ReplayEventListItem = MoveTemp(EventData);
 							}
 						}
 
@@ -1498,7 +1618,7 @@ void FLocalFileNetworkReplayStreamer::FlushStream(const uint32 TimeInMS)
 	UE_LOG(LogLocalFileReplay, Verbose, TEXT("FLocalFileNetworkReplayStreamer::FlushStream. StreamChunkIndex: %i, Size: %i"), StreamChunkIndex, StreamAr.Buffer.Num());
 
 	AddGenericRequestToQueue<FLocalFileReplayInfo>(EQueuedLocalFileRequestType::WritingStream, 
-		[this, StreamChunkStartMS, StreamChunkEndMS, StreamData=MoveTemp(StreamAr.Buffer)](FLocalFileReplayInfo& ReplayInfo) mutable
+		[this, StreamChunkStartMS, StreamChunkEndMS, StreamData=MoveTemp(StreamAr.Buffer), EncryptionKey=CurrentReplayInfo.EncryptionKey](FLocalFileReplayInfo& ReplayInfo) mutable
 		{
 			SCOPE_CYCLE_COUNTER(STAT_LocalReplay_FlushStream);
 
@@ -1507,13 +1627,15 @@ void FLocalFileNetworkReplayStreamer::FlushStream(const uint32 TimeInMS)
 			{
 				LocalFileAr->Seek(LocalFileAr->TotalSize());
 
-				TArray<uint8> FinalData;
+				int32 SizeInMemory = StreamData.Num();
+
+				TArray<uint8> CompressedData;
 
 				if (SupportsCompression())
 				{
 					SCOPE_CYCLE_COUNTER(STAT_LocalReplay_CompressTime);
 
-					if (!CompressBuffer(StreamData, FinalData))
+					if (!CompressBuffer(StreamData, CompressedData))
 					{
 						UE_LOG(LogLocalFileReplay, Warning, TEXT("FLocalFileNetworkReplayStreamer::FlushStream - CompressBuffer failed"));
 						ReplayInfo.bIsValid = false;
@@ -1522,11 +1644,29 @@ void FLocalFileNetworkReplayStreamer::FlushStream(const uint32 TimeInMS)
 				}
 				else
 				{
-					FinalData = MoveTemp(StreamData);
+					CompressedData = MoveTemp(StreamData);
+				}
+
+				TArray<uint8> EncryptedData;
+
+				if (SupportsEncryption())
+				{
+					SCOPE_CYCLE_COUNTER(STAT_LocalReplay_EncryptTime);
+
+					if (!EncryptBuffer(CompressedData, EncryptedData, EncryptionKey))
+					{
+						UE_LOG(LogLocalFileReplay, Warning, TEXT("FLocalFileNetworkReplayStreamer::FlushStream - EncryptBuffer failed"));
+						ReplayInfo.bIsValid = false;
+						return;
+					}
+				}
+				else
+				{
+					EncryptedData = MoveTemp(CompressedData);
 				}
 
 				// flush chunk to disk
-				if (FinalData.Num() > 0)
+				if (EncryptedData.Num() > 0)
 				{
 					ELocalFileChunkType ChunkType = ELocalFileChunkType::ReplayData;
 					*LocalFileAr << ChunkType;
@@ -1544,10 +1684,12 @@ void FLocalFileNetworkReplayStreamer::FlushStream(const uint32 TimeInMS)
 					uint32 Time2 = StreamChunkEndMS;
 					*LocalFileAr << Time2;
 
-					int32 DataSize = FinalData.Num();
-					*LocalFileAr << DataSize;
+					int32 SizeOnDisk = EncryptedData.Num();
+					*LocalFileAr << SizeOnDisk;
 
-					LocalFileAr->Serialize((void*)FinalData.GetData(), FinalData.Num());
+					*LocalFileAr << SizeInMemory;
+
+					LocalFileAr->Serialize((void*)EncryptedData.GetData(), EncryptedData.Num());
 
 					int32 ChunkSize = LocalFileAr->Tell() - MetadataPos;
 
@@ -1564,9 +1706,7 @@ void FLocalFileNetworkReplayStreamer::FlushStream(const uint32 TimeInMS)
 		{
 			if (ReplayInfo.bIsValid)
 			{
-				const int32 TotalLengthInMS = CurrentReplayInfo.LengthInMS;
-				CurrentReplayInfo = ReplayInfo;
-				CurrentReplayInfo.LengthInMS = TotalLengthInMS;
+				UpdateCurrentReplayInfo(ReplayInfo);
 			}
 			else
 			{
@@ -1614,7 +1754,7 @@ void FLocalFileNetworkReplayStreamer::FlushCheckpointInternal(const uint32 TimeI
 	const uint32 CheckpointTimeInMS = StreamTimeRange.Max;
 
 	AddGenericRequestToQueue<FLocalFileReplayInfo>(EQueuedLocalFileRequestType::WritingCheckpoint, 
-		[this, CheckpointTimeInMS, TotalLengthInMS, CheckpointData=MoveTemp(CheckpointAr.Buffer)](FLocalFileReplayInfo& ReplayInfo) mutable
+		[this, CheckpointTimeInMS, TotalLengthInMS, CheckpointData=MoveTemp(CheckpointAr.Buffer), EncryptionKey=CurrentReplayInfo.EncryptionKey](FLocalFileReplayInfo& ReplayInfo) mutable
 		{
 			SCOPE_CYCLE_COUNTER(STAT_LocalReplay_FlushCheckpoint);
 
@@ -1628,13 +1768,13 @@ void FLocalFileNetworkReplayStreamer::FlushCheckpointInternal(const uint32 TimeI
 				{
 					LocalFileAr->Seek(LocalFileAr->TotalSize());
 
-					TArray<uint8> FinalData;
+					TArray<uint8> CompressedData;
 
 					if (SupportsCompression())
 					{
 						SCOPE_CYCLE_COUNTER(STAT_LocalReplay_CompressTime);
 
-						if (!CompressBuffer(CheckpointData, FinalData))
+						if (!CompressBuffer(CheckpointData, CompressedData))
 						{
 							UE_LOG(LogLocalFileReplay, Warning, TEXT("FLocalFileNetworkReplayStreamer::FlushStream - CompressBuffer failed"));
 							ReplayInfo.bIsValid = false;
@@ -1643,11 +1783,29 @@ void FLocalFileNetworkReplayStreamer::FlushCheckpointInternal(const uint32 TimeI
 					}
 					else
 					{
-						FinalData = MoveTemp(CheckpointData);
+						CompressedData = MoveTemp(CheckpointData);
+					}
+
+					TArray<uint8> EncryptedData;
+
+					if (SupportsEncryption())
+					{
+						SCOPE_CYCLE_COUNTER(STAT_LocalReplay_EncryptTime);
+
+						if (!EncryptBuffer(CompressedData, EncryptedData, EncryptionKey))
+						{
+							UE_LOG(LogLocalFileReplay, Warning, TEXT("FLocalFileNetworkReplayStreamer::FlushStream - EncryptBuffer failed"));
+							ReplayInfo.bIsValid = false;
+							return;
+						}
+					}
+					else
+					{
+						EncryptedData = MoveTemp(CompressedData);
 					}
 
 					// flush checkpoint
-					if (FinalData.Num() > 0)
+					if (EncryptedData.Num() > 0)
 					{
 						ELocalFileChunkType ChunkType = ELocalFileChunkType::Checkpoint;
 						*LocalFileAr << ChunkType;
@@ -1674,10 +1832,10 @@ void FLocalFileNetworkReplayStreamer::FlushCheckpointInternal(const uint32 TimeI
 						uint32 Time2 = CheckpointTimeInMS;
 						*LocalFileAr << Time2;
 
-						int32 CheckpointSize = FinalData.Num();
+						int32 CheckpointSize = EncryptedData.Num();
 						*LocalFileAr << CheckpointSize;
 
-						LocalFileAr->Serialize((void*)FinalData.GetData(), FinalData.Num());
+						LocalFileAr->Serialize((void*)EncryptedData.GetData(), EncryptedData.Num());
 
 						int32 ChunkSize = LocalFileAr->Tell() - MetadataPos;
 
@@ -1691,6 +1849,7 @@ void FLocalFileNetworkReplayStreamer::FlushCheckpointInternal(const uint32 TimeI
 				if (ReadReplayInfo(CurrentStreamName, ReplayInfo))
 				{
 					ReplayInfo.LengthInMS = TotalLengthInMS;
+					ReplayInfo.EncryptionKey = EncryptionKey;
 
 					WriteReplayInfo(CurrentStreamName, ReplayInfo);
 				}
@@ -1700,9 +1859,7 @@ void FLocalFileNetworkReplayStreamer::FlushCheckpointInternal(const uint32 TimeI
 		{
 			if (ReplayInfo.bIsValid)
 			{
-				int32 CurrentTotalLengthInMS = CurrentReplayInfo.LengthInMS;
-				CurrentReplayInfo = ReplayInfo;
-				CurrentReplayInfo.LengthInMS = CurrentTotalLengthInMS;
+				UpdateCurrentReplayInfo(ReplayInfo);
 			}
 			else
 			{
@@ -1823,29 +1980,48 @@ void FLocalFileNetworkReplayStreamer::GotoCheckpointIndex(const int32 Checkpoint
 
 							LocalFileAr->Serialize(CheckpointData.GetData(), CheckpointData.Num());
 
+							TArray<uint8> PlaintextData;
+
 							// Get the checkpoint data
+							if (RequestData.ReplayInfo.bEncrypted)
+							{
+								if (SupportsEncryption())
+								{
+									SCOPE_CYCLE_COUNTER(STAT_LocalReplay_DecryptTime);
+
+									if (!DecryptBuffer(CheckpointData, PlaintextData, RequestData.ReplayInfo.EncryptionKey))
+									{
+										UE_LOG(LogLocalFileReplay, Error, TEXT("FLocalFileNetworkReplayStreamer::GotoCheckpointIndexDelta. DecryptBuffer FAILED."));
+										RequestData.DataBuffer.Empty();
+										return;
+									}
+								}
+								else
+								{
+									UE_LOG(LogLocalFileReplay, Error, TEXT("FLocalFileNetworkReplayStreamer::GotoCheckpointIndexDelta. Encrypted checkpoint but streamer does not support encryption."));
+									RequestData.DataBuffer.Empty();
+									return;
+								}
+							}
+							else
+							{
+								PlaintextData = MoveTemp(CheckpointData);
+							}
+
+							TArray<uint8> UncompressedData;
+
 							if (RequestData.ReplayInfo.bCompressed)
 							{
 								if (SupportsCompression())
 								{
 									SCOPE_CYCLE_COUNTER(STAT_LocalReplay_DecompressTime);
 
-									TArray<uint8> UncompressedData;
-
-									if (!DecompressBuffer(CheckpointData, UncompressedData))
+									if (!DecompressBuffer(PlaintextData, UncompressedData))
 									{
 										UE_LOG(LogLocalFileReplay, Error, TEXT("FLocalFileNetworkReplayStreamer::GotoCheckpointIndexDelta. DecompressBuffer FAILED."));
 										RequestData.DataBuffer.Empty();
 										return;
 									}
-
-									FMemoryWriter Writer(RequestData.DataBuffer, true, true);
-									uint32 CheckpointSize = UncompressedData.Num();
-									Writer << CheckpointSize;
-
-									RequestData.DataBuffer.Append(UncompressedData);
-
-									DeltaCheckpointCache.Add(i, MakeShareable(new FCachedFileRequest(UncompressedData, 0)));
 								}
 								else
 								{
@@ -1856,14 +2032,16 @@ void FLocalFileNetworkReplayStreamer::GotoCheckpointIndex(const int32 Checkpoint
 							}
 							else
 							{
-								FMemoryWriter Writer(RequestData.DataBuffer, true, true);
-								uint32 CheckpointSize = CheckpointData.Num();
-								Writer << CheckpointSize;
+								UncompressedData = MoveTemp(PlaintextData);
+							}
 
-								RequestData.DataBuffer.Append(CheckpointData);
+							FMemoryWriter Writer(RequestData.DataBuffer, true, true);
+							uint32 CheckpointSize = UncompressedData.Num();
+							Writer << CheckpointSize;
 
-								DeltaCheckpointCache.Add(i, MakeShareable(new FCachedFileRequest(CheckpointData, 0)));
-							}
+							RequestData.DataBuffer.Append(UncompressedData);
+
+							DeltaCheckpointCache.Add(i, MakeShareable(new FCachedFileRequest(UncompressedData, 0)));
 						}
 					}
 				}
@@ -1985,6 +2163,32 @@ void FLocalFileNetworkReplayStreamer::GotoCheckpointIndex(const int32 Checkpoint
 					LocalFileAr->Serialize(RequestData.DataBuffer.GetData(), RequestData.DataBuffer.Num());
 
 					// Get the checkpoint data
+					
+					if (RequestData.ReplayInfo.bEncrypted)
+					{
+						if (SupportsEncryption())
+						{
+							SCOPE_CYCLE_COUNTER(STAT_LocalReplay_DecryptTime);
+
+							TArray<uint8> DecryptedData;
+
+							if (!DecryptBuffer(RequestData.DataBuffer, DecryptedData, RequestData.ReplayInfo.EncryptionKey))
+							{
+								UE_LOG(LogLocalFileReplay, Error, TEXT("FLocalFileNetworkReplayStreamer::GotoCheckpointIndex. DecryptBuffer FAILED."));
+								RequestData.DataBuffer.Empty();
+								return;
+							}
+
+							RequestData.DataBuffer = MoveTemp(DecryptedData);
+						}
+						else
+						{
+							UE_LOG(LogLocalFileReplay, Error, TEXT("FLocalFileNetworkReplayStreamer::GotoCheckpointIndex. Encrypted checkpoint but streamer does not support encryption."));
+							RequestData.DataBuffer.Empty();
+							return;
+						}
+					}
+
 					if (RequestData.ReplayInfo.bCompressed)
 					{
 						if (SupportsCompression())
@@ -2014,7 +2218,7 @@ void FLocalFileNetworkReplayStreamer::GotoCheckpointIndex(const int32 Checkpoint
 				LocalFileAr = nullptr;
 			}
 		},
-			[this, CheckpointIndex, Delegate](TLocalFileRequestCommonData<FGotoResult>& RequestData)
+		[this, CheckpointIndex, Delegate](TLocalFileRequestCommonData<FGotoResult>& RequestData)
 		{
 			if (bStopStreamingCalled)
 			{
@@ -2554,9 +2758,7 @@ void FLocalFileNetworkReplayStreamer::WriteHeader()
 		{
 			if (ReplayInfo.bIsValid)
 			{
-				int32 TotalLengthInMS = CurrentReplayInfo.LengthInMS;
-				CurrentReplayInfo = ReplayInfo;
-				CurrentReplayInfo.LengthInMS = TotalLengthInMS;
+				UpdateCurrentReplayInfo(ReplayInfo);
 			}
 			else
 			{
@@ -2719,6 +2921,32 @@ void FLocalFileNetworkReplayStreamer::ConditionallyLoadNextChunk()
 
 					LocalFileAr->Serialize(RequestData.DataBuffer.GetData(), RequestData.DataBuffer.Num());
 
+					if (RequestData.ReplayInfo.bEncrypted)
+					{
+						if (SupportsEncryption())
+						{
+							SCOPE_CYCLE_COUNTER(STAT_LocalReplay_DecryptTime);
+
+							TArray<uint8> DecryptedData;
+							if (DecryptBuffer(RequestData.DataBuffer, DecryptedData, RequestData.ReplayInfo.EncryptionKey))
+							{
+								RequestData.DataBuffer = MoveTemp(DecryptedData);
+							}
+							else
+							{
+								UE_LOG(LogLocalFileReplay, Error, TEXT("ConditionallyLoadNextChunk failed to decrypt data."));
+								RequestData.DataBuffer.Empty();
+								return;
+							}
+						}
+						else
+						{
+							UE_LOG(LogLocalFileReplay, Error, TEXT("ConditionallyLoadNextChunk: Replay is marked encrypted but streamer does not support it."));
+							RequestData.DataBuffer.Empty();
+							return;
+						}
+					}
+
 					if (RequestData.ReplayInfo.bCompressed)
 					{
 						if (SupportsCompression())
@@ -2732,12 +2960,14 @@ void FLocalFileNetworkReplayStreamer::ConditionallyLoadNextChunk()
 							}
 							else
 							{
+								UE_LOG(LogLocalFileReplay, Error, TEXT("ConditionallyLoadNextChunk failed to uncompresss data."));
 								RequestData.DataBuffer.Empty();
 								return;
 							}
 						}
 						else
 						{
+							UE_LOG(LogLocalFileReplay, Error, TEXT("ConditionallyLoadNextChunk: Replay is marked compressed but streamer does not support it."));
 							RequestData.DataBuffer.Empty();
 							return;
 						}
@@ -2986,6 +3216,28 @@ const int32 FLocalFileNetworkReplayStreamer::GetUserIndexFromUserString(const FS
 	return INDEX_NONE;
 }
 
+void FLocalFileNetworkReplayStreamer::UpdateCurrentReplayInfo(FLocalFileReplayInfo& ReplayInfo)
+{
+	if (ensure(ReplayInfo.bIsValid))
+	{
+		// maintain the current values of the total length and encryption key
+		TArray<uint8> CurrentKey = CurrentReplayInfo.EncryptionKey;
+		const int32 TotalLengthInMS = CurrentReplayInfo.LengthInMS;
+
+		CurrentReplayInfo = ReplayInfo;
+		
+		CurrentReplayInfo.LengthInMS = TotalLengthInMS;
+		CurrentReplayInfo.EncryptionKey = MoveTemp(CurrentKey);
+	}
+}
+
+PRAGMA_DISABLE_DEPRECATION_WARNINGS
+int32 FLocalFileNetworkReplayStreamer::GetDecompressedSize(FArchive& InCompressed) const
+{
+	return 0;
+}
+PRAGMA_ENABLE_DEPRECATION_WARNINGS
+
 IMPLEMENT_MODULE(FLocalFileNetworkReplayStreamingFactory, LocalFileNetworkReplayStreaming)
 
 TSharedPtr<INetworkReplayStreamer> FLocalFileNetworkReplayStreamingFactory::CreateReplayStreamer() 
diff --git a/Engine/Source/Runtime/NetworkReplayStreaming/LocalFileNetworkReplayStreaming/Public/LocalFileNetworkReplayStreaming.h b/Engine/Source/Runtime/NetworkReplayStreaming/LocalFileNetworkReplayStreaming/Public/LocalFileNetworkReplayStreaming.h
--- a/Engine/Source/Runtime/NetworkReplayStreaming/LocalFileNetworkReplayStreaming/Public/LocalFileNetworkReplayStreaming.h
+++ b/Engine/Source/Runtime/NetworkReplayStreaming/LocalFileNetworkReplayStreaming/Public/LocalFileNetworkReplayStreaming.h
@@ -49,6 +49,7 @@ struct FLocalFileReplayDataInfo
 		Time1(0),
 		Time2(0),
 		SizeInBytes(0),
+		MemorySizeInBytes(0),
 		ReplayDataOffset(0),
 		StreamOffset(0)
 	{}
@@ -57,6 +58,7 @@ struct FLocalFileReplayDataInfo
 	uint32 Time1;
 	uint32 Time2;
 	int32 SizeInBytes;
+	int32 MemorySizeInBytes;
 	int64 ReplayDataOffset;
 	int64 StreamOffset;
 };
@@ -96,6 +98,7 @@ struct FLocalFileReplayInfo
 		bIsLive(false),
 		bIsValid(false),
 		bCompressed(false),
+		bEncrypted(false),
 		HeaderChunkIndex(INDEX_NONE)
 	{}
 
@@ -109,6 +112,9 @@ struct FLocalFileReplayInfo
 	bool bIsValid;
 	bool bCompressed;
 
+	bool bEncrypted;
+	TArray<uint8> EncryptionKey;
+
 	int32 HeaderChunkIndex;
 
 	TArray<FLocalFileChunkInfo> Chunks;
@@ -450,9 +456,17 @@ public:
 	virtual const int32 GetUserIndexFromUserString(const FString& UserString) override;
 
 	virtual bool SupportsCompression() const { return false; }
-	virtual int32 GetDecompressedSize(FArchive& InCompressed) const { return 0; }
-	virtual bool DecompressBuffer(const TArray<uint8>& InCompressed, TArray< uint8 >& OutBuffer) const { return false; }
-	virtual bool CompressBuffer(const TArray< uint8 >& InBuffer, TArray< uint8 >& OutCompressed) const { return false; }
+
+	UE_DEPRECATED(4.25, "No longer used")
+	virtual int32 GetDecompressedSize(FArchive& InCompressed) const;
+
+	virtual bool DecompressBuffer(const TArray<uint8>& InCompressed, TArray<uint8>& OutBuffer) const { return false; }
+	virtual bool CompressBuffer(const TArray<uint8>& InBuffer, TArray<uint8>& OutCompressed) const { return false; }
+
+	virtual bool SupportsEncryption() const { return false; }
+	virtual void GenerateEncryptionKey(TArray<uint8>& EncryptionKey) {}
+	virtual bool EncryptBuffer(TArrayView<const uint8> Plaintext, TArray<uint8>& Ciphertext, TArrayView<const uint8> EncryptionKey) const { return false; }
+	virtual bool DecryptBuffer(TArrayView<const uint8> Ciphertext, TArray<uint8>& Plaintext, TArrayView<const uint8> EncryptionKey) const { return false; }
 
 	void Tick(float DeltaSeconds);
 
@@ -630,6 +644,8 @@ protected:
 	mutable TMap<FString, TArray<uint8>> FileContentsCache;
 	const TArray<uint8>& GetCachedFileContents(const FString& Filename) const;
 
+	void UpdateCurrentReplayInfo(FLocalFileReplayInfo& ReplayInfo);
+
 public:
 	static const FString& GetDefaultDemoSavePath();
 
diff --git a/Engine/Source/Runtime/NetworkReplayStreaming/SaveGameNetworkReplayStreaming/Private/SaveGameNetworkReplayStreaming.cpp b/Engine/Source/Runtime/NetworkReplayStreaming/SaveGameNetworkReplayStreaming/Private/SaveGameNetworkReplayStreaming.cpp
--- a/Engine/Source/Runtime/NetworkReplayStreaming/SaveGameNetworkReplayStreaming/Private/SaveGameNetworkReplayStreaming.cpp
+++ b/Engine/Source/Runtime/NetworkReplayStreaming/SaveGameNetworkReplayStreaming/Private/SaveGameNetworkReplayStreaming.cpp
@@ -1579,6 +1579,27 @@ bool FSaveGameNetworkReplayStreamer::ReadMetaDataFromLocalStream(FArchive& Strea
 			ReplayEventData.SetNumUninitialized(LocalEvent.SizeInBytes);
 			StreamArchive.Seek(LocalEvent.EventDataOffset);
 			StreamArchive.Serialize(ReplayEventData.GetData(), ReplayEventData.Num());
+
+			if (FileReplayInfo.bEncrypted)
+			{
+				if (SupportsEncryption())
+				{
+					TArray<uint8> PlaintextData;
+										
+					if (!DecryptBuffer(ReplayEventData, PlaintextData, FileReplayInfo.EncryptionKey))
+					{
+						UE_LOG(LogSaveGameReplay, Error, TEXT("FSaveGameNetworkReplayStreamer::ReadMetaDataFromLocalStream. DecryptBuffer failed."));
+						return false;
+					}
+
+					ReplayEventData = MoveTemp(PlaintextData);
+				}
+				else
+				{
+					UE_LOG(LogSaveGameReplay, Error, TEXT("FSaveGameNetworkReplayStreamer::ReadMetaDataFromLocalStream. Encrypted event but streamer does not support encryption."));
+					return false;
+				}
+			}
 		}
 
 		EventList.Add(MoveTemp(ReplayEvent));
