Add check for valid replay file name in StartStreaming

diff --git a/Engine/Source/Runtime/NetworkReplayStreaming/LocalFileNetworkReplayStreaming/Private/LocalFileNetworkReplayStreaming.cpp b/Engine/Source/Runtime/NetworkReplayStreaming/LocalFileNetworkReplayStreaming/Private/LocalFileNetworkReplayStreaming.cpp
--- a/Engine/Source/Runtime/NetworkReplayStreaming/LocalFileNetworkReplayStreaming/Private/LocalFileNetworkReplayStreaming.cpp
+++ b/Engine/Source/Runtime/NetworkReplayStreaming/LocalFileNetworkReplayStreaming/Private/LocalFileNetworkReplayStreaming.cpp
@@ -682,6 +682,14 @@ void FLocalFileNetworkReplayStreamer::StartStreaming(const FStartStreamingParame
 
 	const FString FullDemoFilename = GetDemoFullFilename(FinalDemoName);
 
+	// only record to valid replay file names
+	if (Params.bRecord && !FullDemoFilename.EndsWith(FNetworkReplayStreaming::GetReplayFileExtension()))
+	{
+		UE_LOG(LogLocalFileReplay, Warning, TEXT("FLocalFileNetworkReplayStreamer::StartStreaming. Invalid replay file name for recording: %s"), *FullDemoFilename);
+		Delegate.ExecuteIfBound(Result);
+		return;
+	}
+
 	CurrentStreamName = FinalDemoName;
 
 	if (!Params.bRecord)
@@ -994,7 +1002,7 @@ void FLocalFileNetworkReplayStreamer::EnumerateStreams(const FNetworkReplayVersi
 		{
 			SCOPE_CYCLE_COUNTER(STAT_LocalReplay_Enumerate);
 
-			const FString WildCardPath = GetDemoPath() + TEXT("*.replay");
+			const FString WildCardPath = GetDemoPath() + TEXT("*") + FNetworkReplayStreaming::GetReplayFileExtension();
 
 			TArray<FString> ReplayFileNames;
 			IFileManager::Get().FindFiles(ReplayFileNames, *WildCardPath, true, false);
@@ -1566,7 +1574,7 @@ void FLocalFileNetworkReplayStreamer::RenameReplay_Internal(const FString& Repla
 			const FString& NewReplayName = GetDemoFullFilename(NewName);
 
 			FString NewReplayBaseName = FPaths::GetBaseFilename(NewReplayName);
-			NewReplayBaseName.RemoveFromEnd(".replay");
+			NewReplayBaseName.RemoveFromEnd(FNetworkReplayStreaming::GetReplayFileExtension());
 
 			// Sanity check to make sure the input name isn't changing directories.
 			if (NewName != NewReplayBaseName)
@@ -2518,7 +2526,7 @@ FString FLocalFileNetworkReplayStreamer::GetDemoFullFilename(const FString& Stre
 	if (FPaths::IsRelative(StreamName))
 	{
 		// Treat relative paths as demo stream names.
-		return FPaths::Combine(*GetDemoPath(), *StreamName) + TEXT(".replay");
+		return FPaths::Combine(*GetDemoPath(), *StreamName) + FNetworkReplayStreaming::GetReplayFileExtension();
 	}
 	else
 	{
@@ -2554,10 +2562,10 @@ FString FLocalFileNetworkReplayStreamer::GetAutomaticDemoName() const
 
 				for (FString& AutoReplay : FoundAutoReplays)
 				{
-					// Convert the replay name to a full path, making sure to remove the extra .replay postfix
+					// Convert the replay name to a full path, making sure to remove the file extension
 					// that GetDemoFullFilename will add.
 					AutoReplay = GetDemoFullFilename(AutoReplay);
-					AutoReplay.RemoveFromEnd(TEXT(".replay"));
+					AutoReplay.RemoveFromEnd(FNetworkReplayStreaming::GetReplayFileExtension());
 
 					FDateTime Timestamp = FileManager.GetTimeStamp(*AutoReplay);
 					if (Timestamp < OldestReplayTimestamp)
diff --git a/Engine/Source/Runtime/NetworkReplayStreaming/NetworkReplayStreaming/Private/NetworkReplayStreaming.cpp b/Engine/Source/Runtime/NetworkReplayStreaming/NetworkReplayStreaming/Private/NetworkReplayStreaming.cpp
--- a/Engine/Source/Runtime/NetworkReplayStreaming/NetworkReplayStreaming/Private/NetworkReplayStreaming.cpp
+++ b/Engine/Source/Runtime/NetworkReplayStreaming/NetworkReplayStreaming/Private/NetworkReplayStreaming.cpp
@@ -76,6 +76,11 @@ FString FNetworkReplayStreaming::GetAutomaticReplayPrefix()
 	return CVarReplayStreamerAutoDemoPrefix.GetValueOnAnyThread();
 }
 
+FString FNetworkReplayStreaming::GetReplayFileExtension()
+{
+	return TEXT(".replay");
+}
+
 bool FNetworkReplayStreaming::UseDateTimeAsAutomaticReplayPostfix()
 {
 	return !!CVarReplayStreamerAutoDemoUseDateTimePostfix.GetValueOnAnyThread();
diff --git a/Engine/Source/Runtime/NetworkReplayStreaming/NetworkReplayStreaming/Public/NetworkReplayStreaming.h b/Engine/Source/Runtime/NetworkReplayStreaming/NetworkReplayStreaming/Public/NetworkReplayStreaming.h
--- a/Engine/Source/Runtime/NetworkReplayStreaming/NetworkReplayStreaming/Public/NetworkReplayStreaming.h
+++ b/Engine/Source/Runtime/NetworkReplayStreaming/NetworkReplayStreaming/Public/NetworkReplayStreaming.h
@@ -693,6 +693,10 @@ public:
 	// If false, it's up to the streamer to determine a proper postfix.
 	static NETWORKREPLAYSTREAMING_API bool UseDateTimeAsAutomaticReplayPostfix();
 
+	// Gets the file extension for replay streamers to use when recording replay files to the local file system.
+	// Includes the "." prefix
+	static NETWORKREPLAYSTREAMING_API FString GetReplayFileExtension();
+
 	// FSelfRegisteringExec interface
 	virtual bool Exec(UWorld* InWorld, const TCHAR* Cmd, FOutputDevice& Ar) override;
 
diff --git a/Engine/Source/Runtime/NetworkReplayStreaming/SaveGameNetworkReplayStreaming/Private/SaveGameNetworkReplayStreaming.cpp b/Engine/Source/Runtime/NetworkReplayStreaming/SaveGameNetworkReplayStreaming/Private/SaveGameNetworkReplayStreaming.cpp
--- a/Engine/Source/Runtime/NetworkReplayStreaming/SaveGameNetworkReplayStreaming/Private/SaveGameNetworkReplayStreaming.cpp
+++ b/Engine/Source/Runtime/NetworkReplayStreaming/SaveGameNetworkReplayStreaming/Private/SaveGameNetworkReplayStreaming.cpp
@@ -326,7 +326,8 @@ namespace SaveGameReplay
 		{
 			IFileManager& FileManager = IFileManager::Get();
 
-			const FString WildCard(FPaths::Combine(DemoPath, FString::Printf(TEXT("*%s.replay"), *SaveReplayExt)));
+			const FString LocalFileExt = FNetworkReplayStreaming::GetReplayFileExtension();
+			const FString WildCard(FPaths::Combine(DemoPath, FString::Printf(TEXT("*%s%s"), *SaveReplayExt, *LocalFileExt)));
 			TArray<FString> FoundFiles;
 
 			FileManager.FindFiles(FoundFiles, *WildCard, /*bFiles=*/true, /*bDirectories=*/false);
@@ -336,7 +337,7 @@ namespace SaveGameReplay
 				UE_LOG(LogSaveGameReplay, Log, TEXT("FSaveGameMoveFileHelper::SanitizeNames - Handling %s"), *CurrentName);
 
 				FString NewName(CurrentName);
-				NewName.RemoveFromEnd(TEXT(".replay"));
+				NewName.RemoveFromEnd(LocalFileExt);
 				MakeUniqueReplayName(NewName);
 
 				if (!FileManager.Move(*FPaths::Combine(DemoPath, NewName), *FPaths::Combine(DemoPath, CurrentName)))
@@ -353,12 +354,14 @@ namespace SaveGameReplay
 			// Make sure to sanitize this so the system doesn't get tricked into thinking this is a saved replay.
 			// Note, this may still happen if the usually manually entered this...
 			Name.RemoveFromEnd(SaveReplayExt);
+			
+			const FString LocalFileExt = FNetworkReplayStreaming::GetReplayFileExtension();
 
 			int32 Index = 1;
-			FString UseName = FString::Printf(TEXT("%s.replay"), *Name);
+			FString UseName = FString::Printf(TEXT("%s%s"), *Name, *LocalFileExt);
 			while (FileManager.FileExists(*UseName))
 			{
-				UseName = FString::Printf(TEXT("%s - %d.replay"), *Name, ++Index);
+				UseName = FString::Printf(TEXT("%s - %d%s"), *Name, ++Index, *LocalFileExt);
 			}
 
 			Name = MoveTemp(UseName);
@@ -376,7 +379,8 @@ namespace SaveGameReplay
 
 		static void CopyFile(const TSharedPtr<FMoveContext>& Context, const FString& BaseFileName)
 		{
-			const FString SourceFileName = FPaths::Combine(Context->SourceDirectory, BaseFileName) + ".replay";
+			const FString LocalFileExt = FNetworkReplayStreaming::GetReplayFileExtension();
+			const FString SourceFileName = FPaths::Combine(Context->SourceDirectory, BaseFileName) + LocalFileExt;
 
 			FString DestinationFileName = FPaths::Combine(Context->DestinationDirectory, BaseFileName);
 			MakeUniqueReplayName(DestinationFileName);
@@ -477,12 +481,14 @@ namespace SaveGameReplay
 
 		static void MoveFilesLocalInternal(const TSharedPtr<FMoveContext>& MoveContext)
 		{
+			const FString LocalFileExt = FNetworkReplayStreaming::GetReplayFileExtension();
+
 			TArray<FString> ReplayFiles;
-			IFileManager::Get().FindFiles(ReplayFiles, *MoveContext->SourceDirectory, TEXT(".replay"));
+			IFileManager::Get().FindFiles(ReplayFiles, *MoveContext->SourceDirectory, *LocalFileExt);
 
 			for (FString& ReplayFileName : ReplayFiles)
 			{
-				ReplayFileName.RemoveFromEnd(TEXT(".replay"));
+				ReplayFileName.RemoveFromEnd(LocalFileExt);
 
 #if PLATFORM_PS4
 				if (ReplayFileName.Compare(ReplayFileName.ToLower(), ESearchCase::CaseSensitive) != 0)
