Removing deprecated replay streamer calls that were using user strings
(and thus, references to GWorld)

diff --git a/Engine/Source/Runtime/NetworkReplayStreaming/HttpNetworkReplayStreaming/Private/HttpNetworkReplayStreaming.cpp b/Engine/Source/Runtime/NetworkReplayStreaming/HttpNetworkReplayStreaming/Private/HttpNetworkReplayStreaming.cpp
--- a/Engine/Source/Runtime/NetworkReplayStreaming/HttpNetworkReplayStreaming/Private/HttpNetworkReplayStreaming.cpp
+++ b/Engine/Source/Runtime/NetworkReplayStreaming/HttpNetworkReplayStreaming/Private/HttpNetworkReplayStreaming.cpp
@@ -1516,18 +1516,6 @@ void FHttpNetworkReplayStreamer::DeleteFinishedStream( const FString& StreamName
 }
 
 void FHttpNetworkReplayStreamer::EnumerateStreams( const FNetworkReplayVersion& InReplayVersion, const int32 UserIndex, const FString& MetaString, const TArray< FString >& ExtraParms, const FEnumerateStreamsCallback& Delegate )
-{
-	EnumerateStreams( InReplayVersion, GetUserStringFromUserIndex(UserIndex), MetaString, ExtraParms, Delegate );
-}
-
-PRAGMA_DISABLE_DEPRECATION_WARNINGS
-
-void FHttpNetworkReplayStreamer::EnumerateStreams( const FNetworkReplayVersion& InReplayVersion, const FString& UserString, const FString& MetaString, const FEnumerateStreamsCallback& Delegate )
-{
-	EnumerateStreams( InReplayVersion, UserString, MetaString, TArray< FString >(), Delegate );
-}
-
-void FHttpNetworkReplayStreamer::EnumerateStreams( const FNetworkReplayVersion& InReplayVersion, const FString& UserString, const FString& MetaString, const TArray< FString >& ExtraParms, const FEnumerateStreamsCallback& Delegate )
 {
 	if ( ServerURL.IsEmpty() )
 	{
@@ -1560,6 +1548,8 @@ void FHttpNetworkReplayStreamer::EnumerateStreams( const FNetworkReplayVersion&
 		URL += FString::Printf( TEXT( "&meta=%s" ), *MetaStringToUse );
 	}
 
+	FString UserString = GetUserStringFromUserIndex(UserIndex);
+
 	// Add optional User parameter (filter replays by a user that was in the replay)
 	if ( !UserString.IsEmpty() )
 	{
@@ -1581,8 +1571,6 @@ void FHttpNetworkReplayStreamer::EnumerateStreams( const FNetworkReplayVersion&
 	AddRequestToQueue( EQueuedHttpRequestType::EnumeratingSessions, HttpRequest );
 }
 
-PRAGMA_ENABLE_DEPRECATION_WARNINGS
-
 void FHttpNetworkReplayStreamer::EnumerateRecentStreams( const FNetworkReplayVersion& InReplayVersion, const int32 UserIndex, const FEnumerateStreamsCallback& Delegate )
 {
 	if ( ServerURL.IsEmpty() )
@@ -2715,28 +2703,6 @@ bool FHttpNetworkReplayStreamer::IsCheckpointTypeSupported(EReplayCheckpointType
 	return bSupported;
 }
 
-const int32 FHttpNetworkReplayStreamer::GetUserIndexFromUserString(const FString& UserString)
-{
-	if (!UserString.IsEmpty() && GEngine != nullptr)
-	{
-		if (UWorld* World = GWorld.GetReference())
-		{
-			for (auto ConstIt = GEngine->GetLocalPlayerIterator(World); ConstIt; ++ConstIt)
-			{
-				if (ULocalPlayer const * const LocalPlayer = *ConstIt)
-				{
-					if (UserString.Equals(LocalPlayer->GetPreferredUniqueNetId().ToString()))
-					{
-						return LocalPlayer->GetControllerId();
-					}
-				}
-			}
-		}
-	}
-
-	return INDEX_NONE;
-}
-
 IMPLEMENT_MODULE( FHttpNetworkReplayStreamingFactory, HttpNetworkReplayStreaming )
 
 TSharedPtr< INetworkReplayStreamer > FHttpNetworkReplayStreamingFactory::CreateReplayStreamer()
diff --git a/Engine/Source/Runtime/NetworkReplayStreaming/HttpNetworkReplayStreaming/Public/HttpNetworkReplayStreaming.h b/Engine/Source/Runtime/NetworkReplayStreaming/HttpNetworkReplayStreaming/Public/HttpNetworkReplayStreaming.h
--- a/Engine/Source/Runtime/NetworkReplayStreaming/HttpNetworkReplayStreaming/Public/HttpNetworkReplayStreaming.h
+++ b/Engine/Source/Runtime/NetworkReplayStreaming/HttpNetworkReplayStreaming/Public/HttpNetworkReplayStreaming.h
@@ -224,8 +224,6 @@ public:
 	virtual bool		IsLive() const override;
 	virtual void		DeleteFinishedStream( const FString& StreamName, const FDeleteFinishedStreamCallback& Delegate ) override;
 	virtual void		DeleteFinishedStream( const FString& StreamName, const int32 UserIndex, const FDeleteFinishedStreamCallback& Delegate ) override;
-	virtual void		EnumerateStreams( const FNetworkReplayVersion& ReplayVersion, const FString& UserString, const FString& MetaString, const FEnumerateStreamsCallback& Delegate ) override;
-	virtual void		EnumerateStreams( const FNetworkReplayVersion& InReplayVersion, const FString& UserString, const FString& MetaString, const TArray< FString >& ExtraParms, const FEnumerateStreamsCallback& Delegate ) override;
 	virtual void		EnumerateStreams( const FNetworkReplayVersion& InReplayVersion, const int32 UserIndex, const FString& MetaString, const TArray< FString >& ExtraParms, const FEnumerateStreamsCallback& Delegate ) override;
 	virtual void		EnumerateEvents( const FString& Group, const FEnumerateEventsCallback& Delegate ) override;
 	virtual void		EnumerateEvents( const FString& ReplayName, const FString& Group, const FEnumerateEventsCallback& Delegate ) override;
@@ -265,8 +263,6 @@ public:
 
 	virtual bool IsCheckpointTypeSupported(EReplayCheckpointType CheckpointType) const override;
 
-	virtual const int32 GetUserIndexFromUserString(const FString& UserString) override;
-
 	/** FHttpNetworkReplayStreamer */
 	void UploadHeader();
 	void FlushStream();
diff --git a/Engine/Source/Runtime/NetworkReplayStreaming/InMemoryNetworkReplayStreaming/Private/InMemoryNetworkReplayStreaming.cpp b/Engine/Source/Runtime/NetworkReplayStreaming/InMemoryNetworkReplayStreaming/Private/InMemoryNetworkReplayStreaming.cpp
--- a/Engine/Source/Runtime/NetworkReplayStreaming/InMemoryNetworkReplayStreaming/Private/InMemoryNetworkReplayStreaming.cpp
+++ b/Engine/Source/Runtime/NetworkReplayStreaming/InMemoryNetworkReplayStreaming/Private/InMemoryNetworkReplayStreaming.cpp
@@ -560,28 +560,6 @@ TStatId FInMemoryNetworkReplayStreamer::GetStatId() const
 	RETURN_QUICK_DECLARE_CYCLE_STAT(FInMemoryNetworkReplayStreamer, STATGROUP_Tickables);
 }
 
-const int32 FInMemoryNetworkReplayStreamer::GetUserIndexFromUserString(const FString& UserString)
-{
-	if (!UserString.IsEmpty() && GEngine != nullptr)
-	{
-		if (UWorld* World = GWorld.GetReference())
-		{
-			for (auto ConstIt = GEngine->GetLocalPlayerIterator(World); ConstIt; ++ConstIt)
-			{
-				if (ULocalPlayer const * const LocalPlayer = *ConstIt)
-				{
-					if (UserString.Equals(LocalPlayer->GetPreferredUniqueNetId().ToString()))
-					{
-						return LocalPlayer->GetControllerId();
-					}
-				}
-			}
-		}
-	}
-
-	return INDEX_NONE;
-}
-
 void FInMemoryReplayStreamArchive::Serialize(void* V, int64 Length) 
 {
 	if (IsLoading() )
diff --git a/Engine/Source/Runtime/NetworkReplayStreaming/InMemoryNetworkReplayStreaming/Public/InMemoryNetworkReplayStreaming.h b/Engine/Source/Runtime/NetworkReplayStreaming/InMemoryNetworkReplayStreaming/Public/InMemoryNetworkReplayStreaming.h
--- a/Engine/Source/Runtime/NetworkReplayStreaming/InMemoryNetworkReplayStreaming/Public/InMemoryNetworkReplayStreaming.h
+++ b/Engine/Source/Runtime/NetworkReplayStreaming/InMemoryNetworkReplayStreaming/Public/InMemoryNetworkReplayStreaming.h
@@ -194,8 +194,6 @@ public:
 
 	virtual bool IsCheckpointTypeSupported(EReplayCheckpointType CheckpointType) const override { return (CheckpointType == EReplayCheckpointType::Full); }
 
-	virtual const int32 GetUserIndexFromUserString(const FString& UserString) override;
-
 private:
 	bool IsNamedStreamLive( const FString& StreamName ) const;
 
diff --git a/Engine/Source/Runtime/NetworkReplayStreaming/LocalFileNetworkReplayStreaming/Private/LocalFileNetworkReplayStreaming.cpp b/Engine/Source/Runtime/NetworkReplayStreaming/LocalFileNetworkReplayStreaming/Private/LocalFileNetworkReplayStreaming.cpp
--- a/Engine/Source/Runtime/NetworkReplayStreaming/LocalFileNetworkReplayStreaming/Private/LocalFileNetworkReplayStreaming.cpp
+++ b/Engine/Source/Runtime/NetworkReplayStreaming/LocalFileNetworkReplayStreaming/Private/LocalFileNetworkReplayStreaming.cpp
@@ -3216,28 +3216,6 @@ bool FLocalFileNetworkReplayStreamer::IsCheckpointTypeSupported(EReplayCheckpoin
 	return bSupported;
 }
 
-const int32 FLocalFileNetworkReplayStreamer::GetUserIndexFromUserString(const FString& UserString)
-{
-	if (!UserString.IsEmpty() && GEngine != nullptr)
-	{
-		if (UWorld* World = GWorld.GetReference())
-		{
-			for (auto ConstIt = GEngine->GetLocalPlayerIterator(World); ConstIt; ++ConstIt)
-			{
-				if (ULocalPlayer const * const LocalPlayer = *ConstIt)
-				{
-					if (UserString.Equals(LocalPlayer->GetPreferredUniqueNetId().ToString()))
-					{
-						return LocalPlayer->GetControllerId();
-					}
-				}
-			}
-		}
-	}
-
-	return INDEX_NONE;
-}
-
 void FLocalFileNetworkReplayStreamer::UpdateCurrentReplayInfo(FLocalFileReplayInfo& ReplayInfo)
 {
 	if (ensure(ReplayInfo.bIsValid))
diff --git a/Engine/Source/Runtime/NetworkReplayStreaming/LocalFileNetworkReplayStreaming/Public/LocalFileNetworkReplayStreaming.h b/Engine/Source/Runtime/NetworkReplayStreaming/LocalFileNetworkReplayStreaming/Public/LocalFileNetworkReplayStreaming.h
--- a/Engine/Source/Runtime/NetworkReplayStreaming/LocalFileNetworkReplayStreaming/Public/LocalFileNetworkReplayStreaming.h
+++ b/Engine/Source/Runtime/NetworkReplayStreaming/LocalFileNetworkReplayStreaming/Public/LocalFileNetworkReplayStreaming.h
@@ -461,8 +461,6 @@ public:
 
 	virtual bool IsCheckpointTypeSupported(EReplayCheckpointType CheckpointType) const override;
 
-	virtual const int32 GetUserIndexFromUserString(const FString& UserString) override;
-
 	virtual bool SupportsCompression() const { return false; }
 
 	UE_DEPRECATED(4.25, "No longer used")
diff --git a/Engine/Source/Runtime/NetworkReplayStreaming/NetworkReplayStreaming/Private/NetworkReplayStreaming.cpp b/Engine/Source/Runtime/NetworkReplayStreaming/NetworkReplayStreaming/Private/NetworkReplayStreaming.cpp
--- a/Engine/Source/Runtime/NetworkReplayStreaming/NetworkReplayStreaming/Private/NetworkReplayStreaming.cpp
+++ b/Engine/Source/Runtime/NetworkReplayStreaming/NetworkReplayStreaming/Private/NetworkReplayStreaming.cpp
@@ -127,48 +127,4 @@ bool FNetworkReplayStreaming::Exec(UWorld* InWorld, const TCHAR* Cmd, FOutputDev
 	}
 
 	return false;
-}
-
-void INetworkReplayStreamer::StartStreaming(const FString& CustomName, const FString& FriendlyName, const TArray<FString>& UserNames, bool bRecord, const FNetworkReplayVersion& ReplayVersion, const FStartStreamingCallback& Delegate)
-{
-	FStartStreamingParameters Params;
-	Params.CustomName = CustomName;
-	Params.FriendlyName = FriendlyName;
-	
-	for (const FString& UserString : UserNames)
-	{
-		Params.UserIndices.Emplace(GetUserIndexFromUserString(UserString));
-	}
-	
-	Params.bRecord = bRecord;
-	Params.ReplayVersion = ReplayVersion;
-
-	StartStreaming(Params, Delegate);
-}
-
-void INetworkReplayStreamer::StartStreaming(const FString& CustomName, const FString& FriendlyName, const TArray<int32>& UserIndices, bool bRecord, const FNetworkReplayVersion& ReplayVersion, const FStartStreamingCallback& Delegate)
-{
-	FStartStreamingParameters Params;
-	Params.CustomName = CustomName;
-	Params.FriendlyName = FriendlyName;
-	Params.UserIndices = UserIndices;
-	Params.bRecord = bRecord;
-	Params.ReplayVersion = ReplayVersion;
-
-	StartStreaming(Params, Delegate);
-}
-
-void INetworkReplayStreamer::EnumerateStreams(const FNetworkReplayVersion& ReplayVersion, const FString& UserString, const FString& MetaString, const FEnumerateStreamsCallback& Delegate)
-{
-	EnumerateStreams(ReplayVersion, GetUserIndexFromUserString(UserString), MetaString, TArray<FString>(), Delegate);
-}
-
-void INetworkReplayStreamer::EnumerateStreams(const FNetworkReplayVersion& ReplayVersion, const FString& UserString, const FString& MetaString, const TArray<FString>& ExtraParms, const FEnumerateStreamsCallback& Delegate)
-{
-	EnumerateStreams(ReplayVersion, GetUserIndexFromUserString(UserString), MetaString, ExtraParms, Delegate);
-}
-
-void INetworkReplayStreamer::EnumerateRecentStreams(const FNetworkReplayVersion& ReplayVersion, const FString& RecentViewer, const FEnumerateStreamsCallback& Delegate)
-{
-	EnumerateRecentStreams(ReplayVersion, GetUserIndexFromUserString(RecentViewer), Delegate);
 }
\ No newline at end of file
diff --git a/Engine/Source/Runtime/NetworkReplayStreaming/NetworkReplayStreaming/Public/NetworkReplayStreaming.h b/Engine/Source/Runtime/NetworkReplayStreaming/NetworkReplayStreaming/Public/NetworkReplayStreaming.h
--- a/Engine/Source/Runtime/NetworkReplayStreaming/NetworkReplayStreaming/Public/NetworkReplayStreaming.h
+++ b/Engine/Source/Runtime/NetworkReplayStreaming/NetworkReplayStreaming/Public/NetworkReplayStreaming.h
@@ -485,12 +485,6 @@ class NETWORKREPLAYSTREAMING_API INetworkReplayStreamer
 public:
 	virtual ~INetworkReplayStreamer() {}
 
-	UE_DEPRECATED(4.23, "StartStreaming now takes a FStartStreamingParameters argument")
-	virtual void StartStreaming(const FString& CustomName, const FString& FriendlyName, const TArray<FString>& UserNames, bool bRecord, const FNetworkReplayVersion& ReplayVersion, const FStartStreamingCallback& Delegate);
-
-	UE_DEPRECATED(4.23, "StartStreaming now takes a FStartStreamingParameters argument")
-	virtual void StartStreaming(const FString& CustomName, const FString& FriendlyName, const TArray<int32>& UserIndices, bool bRecord, const FNetworkReplayVersion& ReplayVersion, const FStartStreamingCallback& Delegate);
-
 	virtual void StartStreaming(const FStartStreamingParameters& Params, const FStartStreamingCallback& Delegate) = 0;
 
 	virtual void StopStreaming() = 0;
@@ -499,12 +493,8 @@ public:
 	virtual FArchive* GetCheckpointArchive() = 0;
 	virtual void FlushCheckpoint(const uint32 TimeInMS) = 0;
 
-	UE_DEPRECATED(4.23, "Please use the version of GotoCheckpointIndex that accepts a FGotoCallback delegate and an EReplayCheckpointType")
-	virtual void GotoCheckpointIndex(const int32 CheckpointIndex, const FGotoCallback& Delegate) { GotoCheckpointIndex(CheckpointIndex, Delegate, EReplayCheckpointType::Full);	}
 	virtual void GotoCheckpointIndex(const int32 CheckpointIndex, const FGotoCallback& Delegate, EReplayCheckpointType CheckpointType) = 0;
 
-	UE_DEPRECATED(4.23, "Please use the version of GotoCheckpointIndex that accepts a FGotoCallback delegate and an EReplayCheckpointType.")
-	virtual void GotoTimeInMS(const uint32 TimeInMS, const FGotoCallback& Delegate) { GotoTimeInMS(TimeInMS, Delegate, EReplayCheckpointType::Full); }
 	virtual void GotoTimeInMS(const uint32 TimeInMS, const FGotoCallback& Delegate, EReplayCheckpointType CheckpointType) = 0;
 
 	virtual bool IsCheckpointTypeSupported(EReplayCheckpointType CheckpointType) const = 0;
@@ -582,20 +572,10 @@ public:
 	virtual void DeleteFinishedStream(const FString& StreamName, const FDeleteFinishedStreamCallback& Delegate) = 0;
 	virtual void DeleteFinishedStream(const FString& StreamName, const int32 UserIndex, const FDeleteFinishedStreamCallback& Delegate) = 0;
 
-	/**
-	 * Retrieves the streams that are available for viewing. May execute asynchronously.
-	 *
-	 * @param Delegate A delegate that will be executed if bound when the list of streams is available
-	 */
-	UE_DEPRECATED(4.23, "UserString is deprecated in favor of UserIndex")
-	virtual void EnumerateStreams(const FNetworkReplayVersion& ReplayVersion, const FString& UserString, const FString& MetaString, const FEnumerateStreamsCallback& Delegate);
-
 	/**
 	* Retrieves the streams that are available for viewing. May execute asynchronously.
 	* Allows the caller to pass in a custom list of query parameters
 	*/
-	UE_DEPRECATED(4.23, "UserString is deprecated in favor of UserIndex")
-	virtual void EnumerateStreams(const FNetworkReplayVersion& ReplayVersion, const FString& UserString, const FString& MetaString, const TArray<FString>& ExtraParms, const FEnumerateStreamsCallback& Delegate);
 	virtual void EnumerateStreams(const FNetworkReplayVersion& ReplayVersion, const int32 UserIndex, const FString& MetaString, const TArray< FString >& ExtraParms, const FEnumerateStreamsCallback& Delegate) = 0;
 
 	/**
@@ -603,8 +583,6 @@ public:
 	 *
 	 * @param Delegate A delegate that will be executed if bound when the list of streams is available
 	 */
-	UE_DEPRECATED(4.23, "RecentViewer is deprecated in favor of UserIndex")
-	virtual void EnumerateRecentStreams(const FNetworkReplayVersion& ReplayVersion, const FString& RecentViewer, const FEnumerateStreamsCallback& Delegate);
 	virtual void EnumerateRecentStreams(const FNetworkReplayVersion& ReplayVersion, const int32 UserIndex, const FEnumerateStreamsCallback& Delegate) = 0;
 
 	/** Returns the last error that occurred while streaming replays */
@@ -644,13 +622,6 @@ public:
 	virtual void SetAnalyticsProvider(TSharedPtr<IAnalyticsProvider>& InProvider) {}
 
 	virtual void Exec(const TCHAR* Cmd, FOutputDevice& Ar) {}
-
-protected:
-	/** Temporary to assist with deprecation of user string apis */
-	virtual const int32 GetUserIndexFromUserString(const FString& UserString) = 0;
-
-private:
-	const void GetUserIndicesFromUserStrings(const TArray<FString>& UserStrings, TArray<int32>& OutUserIndices);
 };
 
 /** Replay streamer factory */
diff --git a/Engine/Source/Runtime/NetworkReplayStreaming/NullNetworkReplayStreaming/Private/NullNetworkReplayStreaming.cpp b/Engine/Source/Runtime/NetworkReplayStreaming/NullNetworkReplayStreaming/Private/NullNetworkReplayStreaming.cpp
--- a/Engine/Source/Runtime/NetworkReplayStreaming/NullNetworkReplayStreaming/Private/NullNetworkReplayStreaming.cpp
+++ b/Engine/Source/Runtime/NetworkReplayStreaming/NullNetworkReplayStreaming/Private/NullNetworkReplayStreaming.cpp
@@ -752,28 +752,6 @@ EStreamingOperationResult FNullNetworkReplayStreamer::GetDemoPath(FString& DemoP
 	return EStreamingOperationResult::Success;
 }
 
-const int32 FNullNetworkReplayStreamer::GetUserIndexFromUserString(const FString& UserString)
-{
-	if (!UserString.IsEmpty() && GEngine != nullptr)
-	{
-		if (UWorld* World = GWorld.GetReference())
-		{
-			for (auto ConstIt = GEngine->GetLocalPlayerIterator(World); ConstIt; ++ConstIt)
-			{
-				if (ULocalPlayer const * const LocalPlayer = *ConstIt)
-				{
-					if (UserString.Equals(LocalPlayer->GetPreferredUniqueNetId().ToString()))
-					{
-						return LocalPlayer->GetControllerId();
-					}
-				}
-			}
-		}
-	}
-
-	return INDEX_NONE;
-}
-
 IMPLEMENT_MODULE( FNullNetworkReplayStreamingFactory, NullNetworkReplayStreaming )
 
 TSharedPtr< INetworkReplayStreamer > FNullNetworkReplayStreamingFactory::CreateReplayStreamer() 
diff --git a/Engine/Source/Runtime/NetworkReplayStreaming/NullNetworkReplayStreaming/Public/NullNetworkReplayStreaming.h b/Engine/Source/Runtime/NetworkReplayStreaming/NullNetworkReplayStreaming/Public/NullNetworkReplayStreaming.h
--- a/Engine/Source/Runtime/NetworkReplayStreaming/NullNetworkReplayStreaming/Public/NullNetworkReplayStreaming.h
+++ b/Engine/Source/Runtime/NetworkReplayStreaming/NullNetworkReplayStreaming/Public/NullNetworkReplayStreaming.h
@@ -115,8 +115,6 @@ public:
 
 	virtual bool IsCheckpointTypeSupported(EReplayCheckpointType CheckpointType) const override { return (CheckpointType == EReplayCheckpointType::Full); }
 
-	virtual const int32 GetUserIndexFromUserString(const FString& UserString) override;
-
 private:
 	bool IsNamedStreamLive( const FString& StreamName ) const;
 
