Log the amount of space available when we begin replay recording

diff --git a/Engine/Source/Runtime/NetworkReplayStreaming/LocalFileNetworkReplayStreaming/Private/LocalFileNetworkReplayStreaming.cpp b/Engine/Source/Runtime/NetworkReplayStreaming/LocalFileNetworkReplayStreaming/Private/LocalFileNetworkReplayStreaming.cpp
--- a/Engine/Source/Runtime/NetworkReplayStreaming/LocalFileNetworkReplayStreaming/Private/LocalFileNetworkReplayStreaming.cpp
+++ b/Engine/Source/Runtime/NetworkReplayStreaming/LocalFileNetworkReplayStreaming/Private/LocalFileNetworkReplayStreaming.cpp
@@ -762,6 +762,13 @@ void FLocalFileNetworkReplayStreamer::StartStreaming(const FStartStreamingParame
 		// We are recording
 		StreamerState = EStreamerState::Recording;
 
+		uint64 TotalDiskSpace = 0;
+		uint64 TotalDiskFreeSpace = 0;
+		if (FPlatformMisc::GetDiskTotalAndFreeSpace(GetDemoPath(), TotalDiskSpace, TotalDiskFreeSpace))
+		{
+			UE_LOG(LogLocalFileReplay, Log, TEXT("Writing replay to '%s' with %.2fMB free"), *GetDemoPath(), (double)TotalDiskFreeSpace / 1024 / 1024);
+		}
+
 		CurrentReplayInfo.EncryptionKey.Reset();
 
 		// generate key now in case any other events are queued during the initial write
