Failing to uncompress/decrypt a stream chunk will now fail playback

diff --git a/Engine/Source/Runtime/NetworkReplayStreaming/LocalFileNetworkReplayStreaming/Private/LocalFileNetworkReplayStreaming.cpp b/Engine/Source/Runtime/NetworkReplayStreaming/LocalFileNetworkReplayStreaming/Private/LocalFileNetworkReplayStreaming.cpp
--- a/Engine/Source/Runtime/NetworkReplayStreaming/LocalFileNetworkReplayStreaming/Private/LocalFileNetworkReplayStreaming.cpp
+++ b/Engine/Source/Runtime/NetworkReplayStreaming/LocalFileNetworkReplayStreaming/Private/LocalFileNetworkReplayStreaming.cpp
@@ -2933,6 +2933,8 @@ void FLocalFileNetworkReplayStreamer::ConditionallyLoadNextChunk()
 			SCOPE_CYCLE_COUNTER(STAT_LocalReplay_ReadStream);
 			LLM_SCOPE(ELLMTag::Replays);
 
+			RequestData.DelegateResult.Result = EStreamingOperationResult::Success;
+
 			if (ReadReplayInfo(CurrentStreamName, RequestData.ReplayInfo))
 			{
 				check(RequestData.ReplayInfo.DataChunks.IsValidIndex(RequestedStreamChunkIndex));
@@ -2965,6 +2967,7 @@ void FLocalFileNetworkReplayStreamer::ConditionallyLoadNextChunk()
 							{
 								UE_LOG(LogLocalFileReplay, Error, TEXT("ConditionallyLoadNextChunk failed to decrypt data."));
 								RequestData.DataBuffer.Empty();
+								RequestData.DelegateResult.Result = EStreamingOperationResult::ReplayCorrupt;
 								return;
 							}
 						}
@@ -2972,6 +2975,7 @@ void FLocalFileNetworkReplayStreamer::ConditionallyLoadNextChunk()
 						{
 							UE_LOG(LogLocalFileReplay, Error, TEXT("ConditionallyLoadNextChunk: Replay is marked encrypted but streamer does not support it."));
 							RequestData.DataBuffer.Empty();
+							RequestData.DelegateResult.Result = EStreamingOperationResult::Unsupported;
 							return;
 						}
 					}
@@ -2991,6 +2995,7 @@ void FLocalFileNetworkReplayStreamer::ConditionallyLoadNextChunk()
 							{
 								UE_LOG(LogLocalFileReplay, Error, TEXT("ConditionallyLoadNextChunk failed to uncompresss data."));
 								RequestData.DataBuffer.Empty();
+								RequestData.DelegateResult.Result = EStreamingOperationResult::ReplayCorrupt;
 								return;
 							}
 						}
@@ -2998,6 +3003,7 @@ void FLocalFileNetworkReplayStreamer::ConditionallyLoadNextChunk()
 						{
 							UE_LOG(LogLocalFileReplay, Error, TEXT("ConditionallyLoadNextChunk: Replay is marked compressed but streamer does not support it."));
 							RequestData.DataBuffer.Empty();
+							RequestData.DelegateResult.Result = EStreamingOperationResult::Unsupported;
 							return;
 						}
 					}
@@ -3010,6 +3016,13 @@ void FLocalFileNetworkReplayStreamer::ConditionallyLoadNextChunk()
 		{
 			LLM_SCOPE(ELLMTag::Replays);
 
+			// Hijacking this error code to indicate a failure in encryption/compression
+			if (!RequestData.DelegateResult.WasSuccessful())
+			{
+				SetLastError(ENetworkReplayError::ServiceUnavailable);
+				return;
+			}
+
 			// Make sure our stream chunk index didn't change under our feet
 			if (RequestedStreamChunkIndex != StreamChunkIndex)
 			{
