Adding"localReplay.AllowEncryptedRecording" in non-shipping builds to forcably disable replay encryption

diff --git a/Engine/Source/Runtime/NetworkReplayStreaming/LocalFileNetworkReplayStreaming/Private/LocalFileNetworkReplayStreaming.cpp b/Engine/Source/Runtime/NetworkReplayStreaming/LocalFileNetworkReplayStreaming/Private/LocalFileNetworkReplayStreaming.cpp
--- a/Engine/Source/Runtime/NetworkReplayStreaming/LocalFileNetworkReplayStreaming/Private/LocalFileNetworkReplayStreaming.cpp
+++ b/Engine/Source/Runtime/NetworkReplayStreaming/LocalFileNetworkReplayStreaming/Private/LocalFileNetworkReplayStreaming.cpp
@@ -65,6 +65,10 @@ namespace LocalFileReplay
 	TAutoConsoleVariable<int32> CVarMaxBufferedStreamChunks(TEXT("localReplay.MaxBufferedStreamChunks"), 10, TEXT(""));
 	TAutoConsoleVariable<int32> CVarAllowLiveStreamDelete(TEXT("localReplay.AllowLiveStreamDelete"), 1, TEXT(""));
 	TAutoConsoleVariable<float> CVarChunkUploadDelayInSeconds(TEXT("localReplay.ChunkUploadDelayInSeconds"), 20.0f, TEXT(""));
+
+#if !UE_BUILD_SHIPPING
+	TAutoConsoleVariable<int32> CVarAllowEncryptedRecording(TEXT("localReplay.AllowEncryptedRecording"), 1, TEXT(""));
+#endif
 };
 
 const uint32 FLocalFileNetworkReplayStreamer::FileMagic = 0x1CA2E27F;
@@ -527,6 +531,19 @@ bool FLocalFileNetworkReplayStreamer::WriteReplayInfo(FArchive& Archive, const F
 	return WriteReplayInfo(Archive, InReplayInfo, DefaultSerializationInfo);
 }
 
+bool FLocalFileNetworkReplayStreamer::AllowEncryptedWrite() const
+{
+	bool bAllowWrite = SupportsEncryption();
+
+#if !UE_BUILD_SHIPPING
+	bAllowWrite = bAllowWrite && (LocalFileReplay::CVarAllowEncryptedRecording.GetValueOnAnyThread() != 0);
+
+	UE_LOG(LogLocalFileReplay, VeryVerbose, TEXT("FLocalFileNetworkReplayStreamer::AllowEncryptedWrite: %s"), *LexToString(bAllowWrite));
+#endif
+
+	return bAllowWrite;
+}
+
 bool FLocalFileNetworkReplayStreamer::WriteReplayInfo(FArchive& Archive, const FLocalFileReplayInfo& InReplayInfo, FLocalFileSerializationInfo& SerializationInfo)
 {
 	if (SerializationInfo.FileVersion < LocalFileReplay::HISTORY_FIXEDSIZE_FRIENDLY_NAME)
@@ -599,7 +616,7 @@ bool FLocalFileNetworkReplayStreamer::WriteReplayInfo(FArchive& Archive, const F
 
 	if (SerializationInfo.FileVersion >= LocalFileReplay::HISTORY_ENCRYPTION)
 	{
-		uint32 Encrypted = SupportsEncryption() ? 1 : 0;
+		uint32 Encrypted = AllowEncryptedWrite() ? 1 : 0;
 		Archive << Encrypted;
 
 		TArray<uint8> KeyToWrite;
@@ -772,7 +789,7 @@ void FLocalFileNetworkReplayStreamer::StartStreaming(const FStartStreamingParame
 		CurrentReplayInfo.EncryptionKey.Reset();
 
 		// generate key now in case any other events are queued during the initial write
-		if (SupportsEncryption())
+		if (AllowEncryptedWrite())
 		{
 			GenerateEncryptionKey(CurrentReplayInfo.EncryptionKey);
 		}
@@ -1140,7 +1157,7 @@ void FLocalFileNetworkReplayStreamer::AddOrUpdateEvent(const FString& Name, cons
 
 					TArray<uint8> EncryptedData;
 
-					if (SupportsEncryption())
+					if (AllowEncryptedWrite())
 					{
 						SCOPE_CYCLE_COUNTER(STAT_LocalReplay_EncryptTime);
 
@@ -1678,7 +1695,7 @@ void FLocalFileNetworkReplayStreamer::FlushStream(const uint32 TimeInMS)
 
 				TArray<uint8> EncryptedData;
 
-				if (SupportsEncryption())
+				if (AllowEncryptedWrite())
 				{
 					SCOPE_CYCLE_COUNTER(STAT_LocalReplay_EncryptTime);
 
@@ -1817,7 +1834,7 @@ void FLocalFileNetworkReplayStreamer::FlushCheckpointInternal(const uint32 TimeI
 
 					TArray<uint8> EncryptedData;
 
-					if (SupportsEncryption())
+					if (AllowEncryptedWrite())
 					{
 						SCOPE_CYCLE_COUNTER(STAT_LocalReplay_EncryptTime);
 
diff --git a/Engine/Source/Runtime/NetworkReplayStreaming/LocalFileNetworkReplayStreaming/Public/LocalFileNetworkReplayStreaming.h b/Engine/Source/Runtime/NetworkReplayStreaming/LocalFileNetworkReplayStreaming/Public/LocalFileNetworkReplayStreaming.h
--- a/Engine/Source/Runtime/NetworkReplayStreaming/LocalFileNetworkReplayStreaming/Public/LocalFileNetworkReplayStreaming.h
+++ b/Engine/Source/Runtime/NetworkReplayStreaming/LocalFileNetworkReplayStreaming/Public/LocalFileNetworkReplayStreaming.h
@@ -474,6 +474,8 @@ public:
 	virtual bool EncryptBuffer(TArrayView<const uint8> Plaintext, TArray<uint8>& Ciphertext, TArrayView<const uint8> EncryptionKey) const { return false; }
 	virtual bool DecryptBuffer(TArrayView<const uint8> Ciphertext, TArray<uint8>& Plaintext, TArrayView<const uint8> EncryptionKey) const { return false; }
 
+	bool AllowEncryptedWrite() const;
+
 	void Tick(float DeltaSeconds);
 
 	virtual uint32 GetMaxFriendlyNameSize() const override;
