Fix to respect the cached in memory chunk loading path, and not generate errors on scrub.

diff --git a/Engine/Source/Runtime/NetworkReplayStreaming/LocalFileNetworkReplayStreaming/Private/LocalFileNetworkReplayStreaming.cpp b/Engine/Source/Runtime/NetworkReplayStreaming/LocalFileNetworkReplayStreaming/Private/LocalFileNetworkReplayStreaming.cpp
--- a/Engine/Source/Runtime/NetworkReplayStreaming/LocalFileNetworkReplayStreaming/Private/LocalFileNetworkReplayStreaming.cpp
+++ b/Engine/Source/Runtime/NetworkReplayStreaming/LocalFileNetworkReplayStreaming/Private/LocalFileNetworkReplayStreaming.cpp
@@ -2950,8 +2950,6 @@ void FLocalFileNetworkReplayStreamer::ConditionallyLoadNextChunk()
 			SCOPE_CYCLE_COUNTER(STAT_LocalReplay_ReadStream);
 			LLM_SCOPE(ELLMTag::Replays);
 
-			RequestData.DelegateResult.Result = EStreamingOperationResult::Success;
-
 			if (ReadReplayInfo(CurrentStreamName, RequestData.ReplayInfo))
 			{
 				check(RequestData.ReplayInfo.DataChunks.IsValidIndex(RequestedStreamChunkIndex));
@@ -2984,7 +2982,7 @@ void FLocalFileNetworkReplayStreamer::ConditionallyLoadNextChunk()
 							{
 								UE_LOG(LogLocalFileReplay, Error, TEXT("ConditionallyLoadNextChunk failed to decrypt data."));
 								RequestData.DataBuffer.Empty();
-								RequestData.DelegateResult.Result = EStreamingOperationResult::ReplayCorrupt;
+								RequestData.bAsyncError = true;
 								return;
 							}
 						}
@@ -2992,7 +2990,7 @@ void FLocalFileNetworkReplayStreamer::ConditionallyLoadNextChunk()
 						{
 							UE_LOG(LogLocalFileReplay, Error, TEXT("ConditionallyLoadNextChunk: Replay is marked encrypted but streamer does not support it."));
 							RequestData.DataBuffer.Empty();
-							RequestData.DelegateResult.Result = EStreamingOperationResult::Unsupported;
+							RequestData.bAsyncError = true;
 							return;
 						}
 					}
@@ -3012,7 +3010,7 @@ void FLocalFileNetworkReplayStreamer::ConditionallyLoadNextChunk()
 							{
 								UE_LOG(LogLocalFileReplay, Error, TEXT("ConditionallyLoadNextChunk failed to uncompresss data."));
 								RequestData.DataBuffer.Empty();
-								RequestData.DelegateResult.Result = EStreamingOperationResult::ReplayCorrupt;
+								RequestData.bAsyncError = true;
 								return;
 							}
 						}
@@ -3020,7 +3018,7 @@ void FLocalFileNetworkReplayStreamer::ConditionallyLoadNextChunk()
 						{
 							UE_LOG(LogLocalFileReplay, Error, TEXT("ConditionallyLoadNextChunk: Replay is marked compressed but streamer does not support it."));
 							RequestData.DataBuffer.Empty();
-							RequestData.DelegateResult.Result = EStreamingOperationResult::Unsupported;
+							RequestData.bAsyncError = true;
 							return;
 						}
 					}
@@ -3034,7 +3032,7 @@ void FLocalFileNetworkReplayStreamer::ConditionallyLoadNextChunk()
 			LLM_SCOPE(ELLMTag::Replays);
 
 			// Hijacking this error code to indicate a failure in encryption/compression
-			if (!RequestData.DelegateResult.WasSuccessful())
+			if (RequestData.bAsyncError)
 			{
 				SetLastError(ENetworkReplayError::ServiceUnavailable);
 				return;
diff --git a/Engine/Source/Runtime/NetworkReplayStreaming/LocalFileNetworkReplayStreaming/Public/LocalFileNetworkReplayStreaming.h b/Engine/Source/Runtime/NetworkReplayStreaming/LocalFileNetworkReplayStreaming/Public/LocalFileNetworkReplayStreaming.h
--- a/Engine/Source/Runtime/NetworkReplayStreaming/LocalFileNetworkReplayStreaming/Public/LocalFileNetworkReplayStreaming.h
+++ b/Engine/Source/Runtime/NetworkReplayStreaming/LocalFileNetworkReplayStreaming/Public/LocalFileNetworkReplayStreaming.h
@@ -376,6 +376,7 @@ public:
 	DelegateResultType DelegateResult;
 	FLocalFileReplayInfo ReplayInfo;
 	TArray<uint8> DataBuffer;
+	bool bAsyncError = false;
 };
 
 template <typename DelegateResultType>
