Initialize Time1 and Time2 to 0 for FReplayEventListItem

diff --git a/Engine/Source/Runtime/NetworkReplayStreaming/NetworkReplayStreaming/Public/NetworkReplayStreaming.h b/Engine/Source/Runtime/NetworkReplayStreaming/NetworkReplayStreaming/Public/NetworkReplayStreaming.h
--- a/Engine/Source/Runtime/NetworkReplayStreaming/NetworkReplayStreaming/Public/NetworkReplayStreaming.h
+++ b/Engine/Source/Runtime/NetworkReplayStreaming/NetworkReplayStreaming/Public/NetworkReplayStreaming.h
@@ -15,7 +15,7 @@ class IAnalyticsProvider;
 class FReplayEventListItem : public FJsonSerializable
 {
 public:
-	FReplayEventListItem() {}
+	FReplayEventListItem() : Time1(0), Time2(0) {}
 	virtual ~FReplayEventListItem() {}
 
 	FString		ID;
