Fix replay cleanup code using the wrong path when cleaning replays at startup on console.  This meant that replays wouldn't be cleaned up until another one was saved.

diff --git a/Engine/Source/Runtime/NetworkReplayStreaming/LocalFileNetworkReplayStreaming/Private/LocalFileNetworkReplayStreaming.cpp b/Engine/Source/Runtime/NetworkReplayStreaming/LocalFileNetworkReplayStreaming/Private/LocalFileNetworkReplayStreaming.cpp
--- a/Engine/Source/Runtime/NetworkReplayStreaming/LocalFileNetworkReplayStreaming/Private/LocalFileNetworkReplayStreaming.cpp
+++ b/Engine/Source/Runtime/NetworkReplayStreaming/LocalFileNetworkReplayStreaming/Private/LocalFileNetworkReplayStreaming.cpp
@@ -2596,7 +2596,7 @@ FString FLocalFileNetworkReplayStreamer::GetDemoFullFilename(const FString& Stre
 	}
 }
 
-bool FLocalFileNetworkReplayStreamer::GetDemoFreeStorageSpace(uint64& DiskFreeSpace, FString DemoPath)
+bool FLocalFileNetworkReplayStreamer::GetDemoFreeStorageSpace(uint64& DiskFreeSpace, const FString& DemoPath)
 {
 #if PLATFORM_USE_PLATFORM_FILE_MANAGED_STORAGE_WRAPPER
 	int64 AllocatedUsed = 0;
@@ -2630,7 +2630,7 @@ bool FLocalFileNetworkReplayStreamer::GetDemoFreeStorageSpace(uint64& DiskFreeSp
 	return true;
 }
 
-bool FLocalFileNetworkReplayStreamer::CleanUpOldReplays(FString DemoPath)
+bool FLocalFileNetworkReplayStreamer::CleanUpOldReplays(const FString& DemoPath)
 {
 	const int32 MaxDemos = FNetworkReplayStreaming::GetMaxNumberOfAutomaticReplays();
 	const bool bUnlimitedDemos = (MaxDemos <= 0);
@@ -2674,7 +2674,7 @@ bool FLocalFileNetworkReplayStreamer::CleanUpOldReplays(FString DemoPath)
 		{
 			// Convert the replay name to a full path, making sure to remove the file extension
 			// that GetDemoFullFilename will add.
-			AutoReplay = GetDemoFullFilename(DemoPath,AutoReplay); // DL: this makes me sick... we know the full path because we passed it in above..
+			AutoReplay = GetDemoFullFilename(DemoPath, AutoReplay); // DL: this makes me sick... we know the full path because we passed it in above..
 			AutoReplay.RemoveFromEnd(FNetworkReplayStreaming::GetReplayFileExtension());
 
 			FDateTime Timestamp = FileManager.GetTimeStamp(*AutoReplay);
@@ -3458,7 +3458,7 @@ void FLocalFileNetworkReplayStreamingFactory::Flush()
 
 void FLocalFileNetworkReplayStreamingFactory::StartupModule()
 {
-	FLocalFileNetworkReplayStreamer::CleanUpOldReplays();
+	FLocalFileNetworkReplayStreamer::CleanUpOldReplays(FLocalFileNetworkReplayStreamer::GetDefaultDemoSavePath());
 }
 
 void FLocalFileNetworkReplayStreamingFactory::ShutdownModule()
diff --git a/Engine/Source/Runtime/NetworkReplayStreaming/LocalFileNetworkReplayStreaming/Public/LocalFileNetworkReplayStreaming.h b/Engine/Source/Runtime/NetworkReplayStreaming/LocalFileNetworkReplayStreaming/Public/LocalFileNetworkReplayStreaming.h
--- a/Engine/Source/Runtime/NetworkReplayStreaming/LocalFileNetworkReplayStreaming/Public/LocalFileNetworkReplayStreaming.h
+++ b/Engine/Source/Runtime/NetworkReplayStreaming/LocalFileNetworkReplayStreaming/Public/LocalFileNetworkReplayStreaming.h
@@ -671,8 +671,8 @@ protected:
 public:
 	static const FString& GetDefaultDemoSavePath();
 	static FString GetDemoFullFilename(const FString& DemoPath, const FString& FileName);
-	static bool CleanUpOldReplays(FString DemoPath = GetDefaultDemoSavePath());
-	static bool GetDemoFreeStorageSpace(uint64& DiskFreeSpace, FString DemoPath);
+	static bool CleanUpOldReplays(const FString& DemoPath = GetDefaultDemoSavePath());
+	static bool GetDemoFreeStorageSpace(uint64& DiskFreeSpace, const FString& DemoPath);
 
 
 
diff --git a/Engine/Source/Runtime/NetworkReplayStreaming/SaveGameNetworkReplayStreaming/Private/SaveGameNetworkReplayStreaming.cpp b/Engine/Source/Runtime/NetworkReplayStreaming/SaveGameNetworkReplayStreaming/Private/SaveGameNetworkReplayStreaming.cpp
--- a/Engine/Source/Runtime/NetworkReplayStreaming/SaveGameNetworkReplayStreaming/Private/SaveGameNetworkReplayStreaming.cpp
+++ b/Engine/Source/Runtime/NetworkReplayStreaming/SaveGameNetworkReplayStreaming/Private/SaveGameNetworkReplayStreaming.cpp
@@ -1815,7 +1815,12 @@ TFunction<bool(const TCHAR*, EGameDelegates_SaveGame, FString&)> FSaveGameNetwor
 
 IMPLEMENT_MODULE(FSaveGameNetworkReplayStreamingFactory, SaveGameNetworkReplayStreaming)
 
-TSharedPtr< INetworkReplayStreamer > FSaveGameNetworkReplayStreamingFactory::CreateReplayStreamer()
+void FSaveGameNetworkReplayStreamingFactory::StartupModule()
+{
+	FSaveGameNetworkReplayStreamer::CleanUpOldReplays(FSaveGameNetworkReplayStreamer::GetDefaultDemoSavePath());
+}
+
+TSharedPtr<INetworkReplayStreamer> FSaveGameNetworkReplayStreamingFactory::CreateReplayStreamer()
 {
 	if (IPlatformFeaturesModule::Get().GetSaveGameSystem() == nullptr)
 	{
diff --git a/Engine/Source/Runtime/NetworkReplayStreaming/SaveGameNetworkReplayStreaming/Public/SaveGameNetworkReplayStreaming.h b/Engine/Source/Runtime/NetworkReplayStreaming/SaveGameNetworkReplayStreaming/Public/SaveGameNetworkReplayStreaming.h
--- a/Engine/Source/Runtime/NetworkReplayStreaming/SaveGameNetworkReplayStreaming/Public/SaveGameNetworkReplayStreaming.h
+++ b/Engine/Source/Runtime/NetworkReplayStreaming/SaveGameNetworkReplayStreaming/Public/SaveGameNetworkReplayStreaming.h
@@ -160,10 +160,14 @@ private:
 	// Although this isn't used on the GameThread, it should only be created / destroyed
 	// from the same thread. Therefore, no need to make it thread safe (for now).
 	mutable TWeakPtr<FSaveGameOptionInfo> WeakOptionInfo;
+
+	friend class FSaveGameNetworkReplayStreamingFactory;
 };
 
 class SAVEGAMENETWORKREPLAYSTREAMING_API FSaveGameNetworkReplayStreamingFactory : public FLocalFileNetworkReplayStreamingFactory
 {
 public:
-	virtual TSharedPtr< INetworkReplayStreamer > CreateReplayStreamer();
+	virtual void StartupModule() override;
+
+	virtual TSharedPtr<INetworkReplayStreamer> CreateReplayStreamer() override;
 };
