Fix some bugs in FLocalFileNetworkReplayStreamer::GetDemoFreeStorageSpace()
- Fix not falling back to a physical space check if the managed storage quota is unlimited.
- Fixed some comparisions where 0 should have been a valid value.
- Fixed returning total physical size instead of free physical size.
- If the demo save directory doesn't exist, create it so we can get sizes even before demos have been saved.

diff --git a/Engine/Source/Runtime/NetworkReplayStreaming/LocalFileNetworkReplayStreaming/Private/LocalFileNetworkReplayStreaming.cpp b/Engine/Source/Runtime/NetworkReplayStreaming/LocalFileNetworkReplayStreaming/Private/LocalFileNetworkReplayStreaming.cpp
--- a/Engine/Source/Runtime/NetworkReplayStreaming/LocalFileNetworkReplayStreaming/Private/LocalFileNetworkReplayStreaming.cpp
+++ b/Engine/Source/Runtime/NetworkReplayStreaming/LocalFileNetworkReplayStreaming/Private/LocalFileNetworkReplayStreaming.cpp
@@ -2602,31 +2602,50 @@ bool FLocalFileNetworkReplayStreamer::GetDemoFreeStorageSpace(uint64& DiskFreeSp
 	int64 AllocatedUsed = 0;
 	int64 AllocatedFree = 0;
 	int64 AllocatedTotal = 0;
-	if (!FPersistentStorageManager::Get().GetPersistentStorageUsage(DemoPath, AllocatedUsed, AllocatedFree, AllocatedTotal))
-		// if we fail to get the persistent storage space from the persistent storage manager then fallback to the total disk freespace check
+	bool bManagedStorageQueryResult = FPersistentStorageManager::Get().GetPersistentStorageUsage(DemoPath, AllocatedUsed, AllocatedFree, AllocatedTotal);
+
+	if (bManagedStorageQueryResult)
 	{
-		UE_LOG(LogLocalFileReplay, Log, TEXT("Failed to get persistent storage useage for %s from the FPeristentStorageManager, falling back to global total disk size"), *DemoPath);
-#endif
-		uint64 TotalDiskSpace = 0;
-		uint64 TotalDiskFreeSpace = 0;
-		if (!FPlatformMisc::GetDiskTotalAndFreeSpace(DemoPath, TotalDiskSpace, TotalDiskFreeSpace))
+		if (AllocatedTotal >= 0) // If total space is < 0, then the storage category is unlimited, so fall back to a physical free space check
 		{
-			// This initial call to GetDiskTotalAndFreeSpace can fail if no replay has been recorded before and the demo folder doesn't exist, so in this case just return the default path
-			UE_LOG(LogLocalFileReplay, Log, TEXT("FLocalFileNetworkReplayStreamer::GetAutomaticDemoName. Unable to determine free space in %s."), *DemoPath);
-			return false;
+			DiskFreeSpace = 0;
+			if (ensure(AllocatedFree >= 0))
+			{
+				DiskFreeSpace = (uint64)AllocatedFree;
+			}
+			return true;
 		}
-		DiskFreeSpace = TotalDiskSpace;
-#if PLATFORM_USE_PLATFORM_FILE_MANAGED_STORAGE_WRAPPER
 	}
 	else
 	{
-		DiskFreeSpace = 0;
-		if (AllocatedFree > 0)
+		UE_LOG(LogLocalFileReplay, Log, TEXT("Failed to get persistent storage useage for %s from the FPeristentStorageManager, falling back to global total disk size"), *DemoPath);
+	}
+#endif
+
+	uint64 TotalDiskSpace = 0;
+	uint64 TotalDiskFreeSpace = 0;
+	bool bActualStorageQueryResult = FPlatformMisc::GetDiskTotalAndFreeSpace(DemoPath, TotalDiskSpace, TotalDiskFreeSpace);
+
+	if (!bActualStorageQueryResult)
+	{
+		// This initial call to GetDiskTotalAndFreeSpace can fail if no replay has been recorded before and the demo folder doesn't exist
+		// so try creating the folder.
+		IFileManager& FileManager = IFileManager::Get();
+		if (!FileManager.DirectoryExists(*DemoPath) && FileManager.MakeDirectory(*DemoPath, true))
 		{
-			DiskFreeSpace = (uint64)AllocatedFree;
+			TotalDiskSpace = 0;
+			TotalDiskFreeSpace = 0;
+			bActualStorageQueryResult = FPlatformMisc::GetDiskTotalAndFreeSpace(DemoPath, TotalDiskSpace, TotalDiskFreeSpace);
+		}
+
+		if (!bActualStorageQueryResult)
+		{
+			UE_LOG(LogLocalFileReplay, Log, TEXT("FLocalFileNetworkReplayStreamer::GetAutomaticDemoName. Unable to determine free space in %s."), *DemoPath);
+			return false;
 		}
 	}
-#endif
+
+	DiskFreeSpace = TotalDiskFreeSpace;
 	return true;
 }
 
