Allow clean-up of temp replay files saved to persistent storage by FSaveGameNetworkReplayStreamer if deleting normal cached replays doesn't free enough space.

diff --git a/Engine/Source/Runtime/NetworkReplayStreaming/LocalFileNetworkReplayStreaming/Private/LocalFileNetworkReplayStreaming.cpp b/Engine/Source/Runtime/NetworkReplayStreaming/LocalFileNetworkReplayStreaming/Private/LocalFileNetworkReplayStreaming.cpp
--- a/Engine/Source/Runtime/NetworkReplayStreaming/LocalFileNetworkReplayStreaming/Private/LocalFileNetworkReplayStreaming.cpp
+++ b/Engine/Source/Runtime/NetworkReplayStreaming/LocalFileNetworkReplayStreaming/Private/LocalFileNetworkReplayStreaming.cpp
@@ -2640,7 +2640,7 @@ bool FLocalFileNetworkReplayStreamer::GetDemoFreeStorageSpace(uint64& DiskFreeSp
 
 		if (!bActualStorageQueryResult)
 		{
-			UE_LOG(LogLocalFileReplay, Log, TEXT("FLocalFileNetworkReplayStreamer::GetAutomaticDemoName. Unable to determine free space in %s."), *DemoPath);
+			UE_LOG(LogLocalFileReplay, Log, TEXT("FLocalFileNetworkReplayStreamer::GetDemoFreeStorageSpace. Unable to determine free space in %s."), *DemoPath);
 			return false;
 		}
 	}
@@ -2649,7 +2649,7 @@ bool FLocalFileNetworkReplayStreamer::GetDemoFreeStorageSpace(uint64& DiskFreeSp
 	return true;
 }
 
-bool FLocalFileNetworkReplayStreamer::CleanUpOldReplays(const FString& DemoPath)
+bool FLocalFileNetworkReplayStreamer::CleanUpOldReplays(const FString& DemoPath, TArrayView<const FString> AdditionalRelativeDemoPaths)
 {
 	const int32 MaxDemos = FNetworkReplayStreaming::GetMaxNumberOfAutomaticReplays();
 	const bool bUnlimitedDemos = (MaxDemos <= 0);
@@ -2670,34 +2670,35 @@ bool FLocalFileNetworkReplayStreamer::CleanUpOldReplays(const FString& DemoPath)
 
 		uint64 MinFreeSpace = LocalFileReplay::CVarReplayRecordingMinSpace.GetValueOnAnyThread();
 
-		const FString WildCardPath = GetDemoFullFilename(DemoPath, AutoPrefix + FString(TEXT("*")));
-
-		TArray<FString> FoundAutoReplays;
-		FileManager.FindFiles(FoundAutoReplays, *WildCardPath, /* bFiles= */ true, /* bDirectories= */ false);
-
 		// build an array of replay info sorted by timestamps
 		struct FAutoReplayInfo
 		{
-			const FString* Path;
-			FDateTime		TimeStamp;
+			FString		Path;
+			FDateTime	TimeStamp;
 
 			// sort by timestamp (reverse order to get newest->oldest)
-			bool operator < (const FAutoReplayInfo& Other) const
+			bool operator<(const FAutoReplayInfo& Other) const
 			{
 				return Other.TimeStamp < TimeStamp;
 			}
 		};
+		
+		// All the replays in the base DemoPath come first
 		TArray<FAutoReplayInfo> SortedAutoReplays;
-		SortedAutoReplays.Reserve(FoundAutoReplays.Num());
-		for (FString& AutoReplay : FoundAutoReplays)
 		{
-			// Convert the replay name to a full path, making sure to remove the file extension
-			// that GetDemoFullFilename will add.
-			AutoReplay = GetDemoFullFilename(DemoPath, AutoReplay); // DL: this makes me sick... we know the full path because we passed it in above..
-			AutoReplay.RemoveFromEnd(FNetworkReplayStreaming::GetReplayFileExtension());
+			const FString WildCardPath = GetDemoFullFilename(DemoPath, AutoPrefix + FString(TEXT("*")));
+
+			TArray<FString> FoundAutoReplays;
+			FileManager.FindFiles(FoundAutoReplays, *WildCardPath, /* bFiles= */ true, /* bDirectories= */ false);
 
-			FDateTime Timestamp = FileManager.GetTimeStamp(*AutoReplay);
-			SortedAutoReplays.Add({ &AutoReplay, Timestamp });
+			SortedAutoReplays.Reserve(SortedAutoReplays.Num() + FoundAutoReplays.Num());
+			for (const FString& AutoReplay : FoundAutoReplays)
+			{
+				// Rebuild full path
+				FString AutoReplayPath = FPaths::Combine(DemoPath, AutoReplay);
+				FDateTime Timestamp = FileManager.GetTimeStamp(*AutoReplayPath);
+				SortedAutoReplays.Add({ MoveTemp(AutoReplayPath), Timestamp });
+			}
 		}
 		SortedAutoReplays.Sort();
 
@@ -2706,24 +2707,69 @@ bool FLocalFileNetworkReplayStreamer::CleanUpOldReplays(const FString& DemoPath)
 			((TotalDiskFreeSpace < MinFreeSpace) || (SortedAutoReplays.Num() >= MaxDemos)))
 		{
 			// find and delete the oldest replay
-			const FString* OldestReplay = SortedAutoReplays[SortedAutoReplays.Num() - 1].Path;
-			SortedAutoReplays.Pop(false);
+			const FAutoReplayInfo OldestReplay = SortedAutoReplays.Pop(false);
+
+			// Try deleting the replay
+			if (!ensureMsgf(FileManager.Delete(*OldestReplay.Path, /*bRequireExists=*/ true, /*bEvenIfReadOnly=*/ true), TEXT("FLocalFileNetworkReplayStreamer::CleanUpOldReplays: Failed to delete old replay %s"), *OldestReplay.Path))
+			{
+				UE_LOG(LogLocalFileReplay, Warning, TEXT("FLocalFileNetworkReplayStreamer::CleanUpOldReplays. Unable to delete old replay %s."), *OldestReplay.Path);
+				return false;
+			}
+
+			if (!GetDemoFreeStorageSpace(TotalDiskFreeSpace, DemoPath))
+			{
+				UE_LOG(LogLocalFileReplay, Warning, TEXT("FLocalFileNetworkReplayStreamer::CleanUpOldReplays. Unable to refresh free space in %s."), *DemoPath);
+				return false;
+			}
+		}
+
+		if (TotalDiskFreeSpace >= MinFreeSpace)
+		{
+			return true;
+		}
+
+		// Only delete additional replays if there is still not enough space
+		TArray<FAutoReplayInfo> SortedAdditionalAutoReplays;
+		for (const FString& RelPath : AdditionalRelativeDemoPaths)
+		{
+			const FString Path = FPaths::Combine(DemoPath, RelPath);
+			const FString AdditionalWildCardPath = GetDemoFullFilename(Path, TEXT("*"));
 
-			check(OldestReplay != nullptr);
+			TArray<FString> FoundAdditionalReplays;
+			FileManager.FindFiles(FoundAdditionalReplays, *AdditionalWildCardPath, /* bFiles= */ true, /* bDirectories= */ false);
 
-			// Try deleting the replay, return an empty string to indicate failure.
-			if (!ensureMsgf(FileManager.Delete(**OldestReplay, /*bRequireExists=*/ true, /*bEvenIfReadOnly=*/ true), TEXT("FLocalFileNetworkReplayStreamer::GetAutomaticDemoName: Failed to delete old replay %s"), **OldestReplay))
+			SortedAdditionalAutoReplays.Reserve(SortedAdditionalAutoReplays.Num() + FoundAdditionalReplays.Num());
+			for (const FString& AutoReplay : FoundAdditionalReplays)
 			{
-				UE_LOG(LogLocalFileReplay, Warning, TEXT("FLocalFileNetworkReplayStreamer::GetAutomaticDemoName. Unable to delete old replay %s."), **OldestReplay);
+				// Rebuild full path
+				FString AutoReplayPath = FPaths::Combine(Path, AutoReplay);
+				FDateTime Timestamp = FileManager.GetTimeStamp(*AutoReplayPath);
+				SortedAdditionalAutoReplays.Add({ MoveTemp(AutoReplayPath), Timestamp });
+			}
+		}
+		SortedAdditionalAutoReplays.Sort();
+
+		// remove oldest replays until we have enough space to record again
+		while (SortedAdditionalAutoReplays.Num() && (TotalDiskFreeSpace < MinFreeSpace))
+		{
+			// find and delete the oldest replay
+			const FAutoReplayInfo OldestReplay = SortedAdditionalAutoReplays.Pop(false);
+
+			// Try deleting the replay
+			if (!ensureMsgf(FileManager.Delete(*OldestReplay.Path, /*bRequireExists=*/ true, /*bEvenIfReadOnly=*/ true), TEXT("FLocalFileNetworkReplayStreamer::CleanUpOldReplays: Failed to delete old replay %s"), *OldestReplay.Path))
+			{
+				UE_LOG(LogLocalFileReplay, Warning, TEXT("FLocalFileNetworkReplayStreamer::CleanUpOldReplays. Unable to delete old replay %s."), *OldestReplay.Path);
 				return false;
 			}
 
 			if (!GetDemoFreeStorageSpace(TotalDiskFreeSpace, DemoPath))
 			{
-				UE_LOG(LogLocalFileReplay, Warning, TEXT("FLocalFileNetworkReplayStreamer::GetAutomaticDemoName. Unable to refresh free space in %s."), *DemoPath);
+				UE_LOG(LogLocalFileReplay, Warning, TEXT("FLocalFileNetworkReplayStreamer::CleanUpOldReplays. Unable to refresh free space in %s."), *DemoPath);
 				return false;
 			}
 		}
+
+		return TotalDiskFreeSpace >= MinFreeSpace;
 	}
 	return true;
 }
@@ -2741,7 +2787,7 @@ FString FLocalFileNetworkReplayStreamer::GetAutomaticDemoName() const
 
 	if (bUseDatePostfix)
 	{
-		if (CleanUpOldReplays(GetDemoPath()) == false)
+		if (CleanUpOldReplays(GetDemoPath(), GetAdditionalRelativeDemoPaths()) == false)
 		{
 			return FString();
 		}
diff --git a/Engine/Source/Runtime/NetworkReplayStreaming/LocalFileNetworkReplayStreaming/Public/LocalFileNetworkReplayStreaming.h b/Engine/Source/Runtime/NetworkReplayStreaming/LocalFileNetworkReplayStreaming/Public/LocalFileNetworkReplayStreaming.h
--- a/Engine/Source/Runtime/NetworkReplayStreaming/LocalFileNetworkReplayStreaming/Public/LocalFileNetworkReplayStreaming.h
+++ b/Engine/Source/Runtime/NetworkReplayStreaming/LocalFileNetworkReplayStreaming/Public/LocalFileNetworkReplayStreaming.h
@@ -625,6 +625,8 @@ protected:
 	virtual TSharedPtr<FArchive> CreateLocalFileWriterForOverwrite(const FString& InFilename) const;
 
 	FString GetDemoPath() const;
+	// Must be relative to the base demo path
+	virtual TArrayView<const FString> GetAdditionalRelativeDemoPaths() const { return {}; }
 	FString GetDemoFullFilename(const FString& FileName) const;
 
 	// Returns a name formatted as "demoX", where X is between 1 and MAX_DEMOS, inclusive.
@@ -671,7 +673,7 @@ protected:
 public:
 	static const FString& GetDefaultDemoSavePath();
 	static FString GetDemoFullFilename(const FString& DemoPath, const FString& FileName);
-	static bool CleanUpOldReplays(const FString& DemoPath = GetDefaultDemoSavePath());
+	static bool CleanUpOldReplays(const FString& DemoPath = GetDefaultDemoSavePath(), TArrayView<const FString> AdditionalRelativeDemoPaths = {});
 	static bool GetDemoFreeStorageSpace(uint64& DiskFreeSpace, const FString& DemoPath);
 
 
diff --git a/Engine/Source/Runtime/NetworkReplayStreaming/SaveGameNetworkReplayStreaming/Private/SaveGameNetworkReplayStreaming.cpp b/Engine/Source/Runtime/NetworkReplayStreaming/SaveGameNetworkReplayStreaming/Private/SaveGameNetworkReplayStreaming.cpp
--- a/Engine/Source/Runtime/NetworkReplayStreaming/SaveGameNetworkReplayStreaming/Private/SaveGameNetworkReplayStreaming.cpp
+++ b/Engine/Source/Runtime/NetworkReplayStreaming/SaveGameNetworkReplayStreaming/Private/SaveGameNetworkReplayStreaming.cpp
@@ -1785,7 +1785,12 @@ FString FSaveGameNetworkReplayStreamer::GetFullPlaybackName() const
 
 FString FSaveGameNetworkReplayStreamer::GetLocalPlaybackName() const
 {
-	return FPaths::Combine(TEXT("Temp/"), PlaybackReplayName);
+	return FPaths::Combine(GetTempDemoRelativeSavePath(), PlaybackReplayName);
+}
+
+TArrayView<const FString> FSaveGameNetworkReplayStreamer::GetAdditionalRelativeDemoPaths() const
+{
+	return MakeArrayView(&GetTempDemoRelativeSavePath(), 1);
 }
 
 const FString& FSaveGameNetworkReplayStreamer::GetDefaultDemoSavePath()
@@ -1794,6 +1799,13 @@ const FString& FSaveGameNetworkReplayStreamer::GetDefaultDemoSavePath()
 	return DefaultDemoSavePath;
 }
 
+const FString& FSaveGameNetworkReplayStreamer::GetTempDemoRelativeSavePath()
+{
+	// Allow for cleanup of the local playback.replay file if needed
+	static const FString TempDemoSavePath(TEXT("Temp/"));
+	return TempDemoSavePath;
+}
+
 const FString& FSaveGameNetworkReplayStreamer::GetDefaultPlaybackName()
 {
 	static const FString DefaultPlaybackName(TEXT("Playback"));
@@ -1817,7 +1829,9 @@ IMPLEMENT_MODULE(FSaveGameNetworkReplayStreamingFactory, SaveGameNetworkReplaySt
 
 void FSaveGameNetworkReplayStreamingFactory::StartupModule()
 {
-	FSaveGameNetworkReplayStreamer::CleanUpOldReplays(FSaveGameNetworkReplayStreamer::GetDefaultDemoSavePath());
+	FSaveGameNetworkReplayStreamer::CleanUpOldReplays(
+		FSaveGameNetworkReplayStreamer::GetDefaultDemoSavePath(),
+		MakeArrayView(&FSaveGameNetworkReplayStreamer::GetTempDemoRelativeSavePath(), 1));
 }
 
 TSharedPtr<INetworkReplayStreamer> FSaveGameNetworkReplayStreamingFactory::CreateReplayStreamer()
diff --git a/Engine/Source/Runtime/NetworkReplayStreaming/SaveGameNetworkReplayStreaming/Public/SaveGameNetworkReplayStreaming.h b/Engine/Source/Runtime/NetworkReplayStreaming/SaveGameNetworkReplayStreaming/Public/SaveGameNetworkReplayStreaming.h
--- a/Engine/Source/Runtime/NetworkReplayStreaming/SaveGameNetworkReplayStreaming/Public/SaveGameNetworkReplayStreaming.h
+++ b/Engine/Source/Runtime/NetworkReplayStreaming/SaveGameNetworkReplayStreaming/Public/SaveGameNetworkReplayStreaming.h
@@ -123,6 +123,8 @@ protected:
 	FString GetFullPlaybackName() const;
 	FString GetLocalPlaybackName() const;
 
+	virtual TArrayView<const FString> GetAdditionalRelativeDemoPaths() const override;
+
 	struct FSaveGameOptionInfo
 	{
 		EGameDelegates_SaveGame Option;
@@ -151,6 +153,7 @@ protected:
 	const FString PlaybackReplayName;
 
 	static const FString& GetDefaultDemoSavePath();
+	static const FString& GetTempDemoRelativeSavePath();
 	static const FString& GetDefaultPlaybackName();
 
 private:
