Adding AppendReplayCommonAttributes to the replay streamer interface that calls new static INetworkReplayStreamer::OnReplayGetAnalyticsAttributes
Replay streamer interface adjustments for common analytics needs

diff --git a/Engine/Source/Runtime/NetworkReplayStreaming/HttpNetworkReplayStreaming/Private/HttpNetworkReplayStreaming.cpp b/Engine/Source/Runtime/NetworkReplayStreaming/HttpNetworkReplayStreaming/Private/HttpNetworkReplayStreaming.cpp
--- a/Engine/Source/Runtime/NetworkReplayStreaming/HttpNetworkReplayStreaming/Private/HttpNetworkReplayStreaming.cpp
+++ b/Engine/Source/Runtime/NetworkReplayStreaming/HttpNetworkReplayStreaming/Private/HttpNetworkReplayStreaming.cpp
@@ -208,7 +208,7 @@ FHttpNetworkReplayStreamer::FHttpNetworkReplayStreamer() :
 	LastChunkTime( 0 ), 
 	LastRefreshViewerTime( 0 ),
 	LastRefreshCheckpointTime( 0 ),
-	StreamerState( EStreamerState::Idle ), 
+	StreamerState(EReplayStreamerState::Idle),
 	bStopStreamingCalled( false ), 
 	bStreamIsLive( false ),
 	NumTotalStreamChunks( 0 ),
@@ -302,7 +302,7 @@ void FHttpNetworkReplayStreamer::StartStreaming(const FStartStreamingParameters&
 	if ( !Params.bRecord )
 	{
 		// We are streaming down
-		StreamerState = EStreamerState::StreamingDown;
+		StreamerState = EReplayStreamerState::Playback;
 
 		SessionName = Params.CustomName;
 
@@ -335,7 +335,7 @@ void FHttpNetworkReplayStreamer::StartStreaming(const FStartStreamingParameters&
 	else
 	{
 		// We are streaming up
-		StreamerState = EStreamerState::StreamingUp;
+		StreamerState = EReplayStreamerState::Recording;
 
 		SessionName.Empty();
 
@@ -553,12 +553,12 @@ void FHttpNetworkReplayStreamer::StopStreaming()
 
 	bStopStreamingCalled = true;
 
-	if ( StreamerState == EStreamerState::StreamingDown )
+	if ( StreamerState == EReplayStreamerState::Playback )
 	{
 		// Refresh one last time, pasing in bFinal == true, this will notify the server we're done (remove us from viewers list)
 		RefreshViewer( true );
 	}
-	else if ( StreamerState == EStreamerState::StreamingUp )
+	else if ( StreamerState == EReplayStreamerState::Recording )
 	{
 		// Flush any final pending stream
 		FlushStream();
@@ -1009,7 +1009,7 @@ void FHttpNetworkReplayStreamer::GotoTimeInMS( const uint32 TimeInMS, const FGot
 
 void FHttpNetworkReplayStreamer::FlushCheckpointInternal( uint32 TimeInMS )
 {
-	if ( SessionName.IsEmpty() || StreamerState != EStreamerState::StreamingUp || CheckpointArchive.Buffer.Num() == 0 )
+	if ( SessionName.IsEmpty() || StreamerState != EReplayStreamerState::Recording || CheckpointArchive.Buffer.Num() == 0 )
 	{
 		// If there is no active session, or we are not recording, we don't need to flush
 		CheckpointArchive.Buffer.Empty();
@@ -1090,7 +1090,7 @@ bool FQueuedHttpRequestAddEvent::PreProcess( FHttpNetworkReplayStreamer* Streame
 
 void FHttpNetworkReplayStreamer::AddOrUpdateEvent( const FString& Name, const uint32 TimeInMS, const FString& Group, const FString& Meta, const TArray<uint8>& Data )
 {
-	if ( StreamerState != EStreamerState::StreamingUp && StreamerState != EStreamerState::StreamingDown )
+	if ( StreamerState != EReplayStreamerState::Recording && StreamerState != EReplayStreamerState::Playback )
 	{
 		UE_LOG( LogHttpReplay, Warning, TEXT( "FHttpNetworkReplayStreamer::AddOrUpdateEvent. Not streaming." ) );
 		return;
@@ -1110,13 +1110,13 @@ void FHttpNetworkReplayStreamer::AddOrUpdateEvent( const FString& Name, const ui
 
 void FHttpNetworkReplayStreamer::AddEvent( const uint32 TimeInMS, const FString& Group, const FString& Meta, const TArray<uint8>& Data )
 {
-	if (StreamerState != EStreamerState::StreamingUp && StreamerState != EStreamerState::StreamingDown)
+	if (StreamerState != EReplayStreamerState::Recording && StreamerState != EReplayStreamerState::Playback)
 	{
 		UE_LOG( LogHttpReplay, Warning, TEXT( "FHttpNetworkReplayStreamer::AddEvent. Not streaming." ) );
 		return;
 	}
 
-	if ( StreamerState == EStreamerState::StreamingUp && !EventGroupSet.Contains( Group ) )
+	if ( StreamerState == EReplayStreamerState::Recording && !EventGroupSet.Contains( Group ) )
 	{
 		// Add the group as a user, so we can quickly find replays that have these event types in them
 		EventGroupSet.Add( Group );
@@ -1406,7 +1406,7 @@ void FHttpNetworkReplayStreamer::CancelStreamingRequests()
 	// Empty the request queue
 	QueuedHttpRequests.Empty();
 
-	StreamerState			= EStreamerState::Idle;
+	StreamerState			= EReplayStreamerState::Idle;
 	bStopStreamingCalled	= false;
 }
 
@@ -1620,7 +1620,7 @@ bool FQueuedHttpRequestAddUser::PreProcess( FHttpNetworkReplayStreamer* Streamer
 
 void FHttpNetworkReplayStreamer::AddUserToReplay(const FString& UserString)
 {
-	if ( StreamerState != EStreamerState::StreamingUp )
+	if ( StreamerState != EReplayStreamerState::Recording )
 	{
 		return;
 	}
@@ -1707,7 +1707,7 @@ void FHttpNetworkReplayStreamer::EnumerateEvents(const FString& ReplayName, cons
 	AddRequestToQueue( EQueuedHttpRequestType::EnumeratingCustomEvent, HttpRequest );
 }
 
-void FHttpNetworkReplayStreamer::RequestFinished( EStreamerState ExpectedStreamerState, EQueuedHttpRequestType::Type ExpectedType, FHttpRequestPtr HttpRequest )
+void FHttpNetworkReplayStreamer::RequestFinished(EReplayStreamerState ExpectedStreamerState, EQueuedHttpRequestType::Type ExpectedType, FHttpRequestPtr HttpRequest )
 {
 	check( StreamerState == ExpectedStreamerState );
 	check( InFlightHttpRequest.IsValid() );
@@ -1721,7 +1721,7 @@ void FHttpNetworkReplayStreamer::HttpStartUploadingFinished( FHttpRequestPtr Htt
 {
 	TSharedPtr< FQueuedHttpRequest > SavedFlightHttpRequest = InFlightHttpRequest;
 
-	RequestFinished( EStreamerState::StreamingUp, EQueuedHttpRequestType::StartUploading, HttpRequest );
+	RequestFinished(EReplayStreamerState::Recording, EQueuedHttpRequestType::StartUploading, HttpRequest );
 
 	if ( bSucceeded && HttpResponse->GetResponseCode() == EHttpResponseCodes::Ok )
 	{
@@ -1759,7 +1759,7 @@ void FHttpNetworkReplayStreamer::HttpStopUploadingFinished( FHttpRequestPtr Http
 {
 	TSharedPtr< FQueuedHttpRequest > SavedFlightHttpRequest = InFlightHttpRequest;
 
-	RequestFinished( EStreamerState::StreamingUp, EQueuedHttpRequestType::StopUploading, HttpRequest );
+	RequestFinished(EReplayStreamerState::Recording, EQueuedHttpRequestType::StopUploading, HttpRequest );
 
 	if ( bSucceeded && HttpResponse->GetResponseCode() == EHttpResponseCodes::NoContent )
 	{
@@ -1788,7 +1788,7 @@ void FHttpNetworkReplayStreamer::HttpHeaderUploadFinished( FHttpRequestPtr HttpR
 {
 	TSharedPtr< FQueuedHttpRequest > SavedFlightHttpRequest = InFlightHttpRequest;
 
-	RequestFinished( EStreamerState::StreamingUp, EQueuedHttpRequestType::UploadingHeader, HttpRequest );
+	RequestFinished(EReplayStreamerState::Recording, EQueuedHttpRequestType::UploadingHeader, HttpRequest );
 
 	FStartStreamingResult Result;
 	Result.bRecording = true;
@@ -1822,7 +1822,7 @@ void FHttpNetworkReplayStreamer::HttpUploadStreamFinished( FHttpRequestPtr HttpR
 {
 	TSharedPtr< FQueuedHttpRequest > SavedFlightHttpRequest = InFlightHttpRequest;
 
-	RequestFinished( EStreamerState::StreamingUp, EQueuedHttpRequestType::UploadingStream, HttpRequest );
+	RequestFinished(EReplayStreamerState::Recording, EQueuedHttpRequestType::UploadingStream, HttpRequest );
 
 	if ( bSucceeded && HttpResponse->GetResponseCode() == EHttpResponseCodes::NoContent )
 	{
@@ -1844,7 +1844,7 @@ void FHttpNetworkReplayStreamer::HttpUploadCheckpointFinished( FHttpRequestPtr H
 {
 	TSharedPtr< FQueuedHttpRequest > SavedFlightHttpRequest = InFlightHttpRequest;
 
-	RequestFinished( EStreamerState::StreamingUp, EQueuedHttpRequestType::UploadingCheckpoint, HttpRequest );
+	RequestFinished(EReplayStreamerState::Recording, EQueuedHttpRequestType::UploadingCheckpoint, HttpRequest );
 
 	if ( bSucceeded && ( HttpResponse->GetResponseCode() == EHttpResponseCodes::Ok || HttpResponse->GetResponseCode() == EHttpResponseCodes::NoContent ) )
 	{
@@ -1888,7 +1888,7 @@ void FHttpNetworkReplayStreamer::HttpUploadCustomEventFinished(FHttpRequestPtr H
 
 void FHttpNetworkReplayStreamer::HttpStartDownloadingFinished( FHttpRequestPtr HttpRequest, FHttpResponsePtr HttpResponse, bool bSucceeded )
 {
-	RequestFinished( EStreamerState::StreamingDown, EQueuedHttpRequestType::StartDownloading, HttpRequest );
+	RequestFinished(EReplayStreamerState::Playback, EQueuedHttpRequestType::StartDownloading, HttpRequest );
 
 	if ( bSucceeded && HttpResponse->GetResponseCode() == EHttpResponseCodes::Ok )
 	{
@@ -1942,7 +1942,7 @@ void FHttpNetworkReplayStreamer::HttpStartDownloadingFinished( FHttpRequestPtr H
 
 void FHttpNetworkReplayStreamer::HttpDownloadHeaderFinished( FHttpRequestPtr HttpRequest, FHttpResponsePtr HttpResponse, bool bSucceeded, FDownloadHeaderCallback Delegate)
 {
-	RequestFinished( EStreamerState::StreamingDown, EQueuedHttpRequestType::DownloadingHeader, HttpRequest );
+	RequestFinished(EReplayStreamerState::Playback, EQueuedHttpRequestType::DownloadingHeader, HttpRequest );
 
 	check( StreamArchive.IsLoading() );
 
@@ -2002,7 +2002,7 @@ void FHttpNetworkReplayStreamer::HttpDownloadFinished( FHttpRequestPtr HttpReque
 {
 	LLM_SCOPE(ELLMTag::Replays);
 
-	RequestFinished( EStreamerState::StreamingDown, EQueuedHttpRequestType::DownloadingStream, HttpRequest );
+	RequestFinished(EReplayStreamerState::Playback, EQueuedHttpRequestType::DownloadingStream, HttpRequest );
 
 	check( StreamArchive.IsLoading() );
 
@@ -2112,7 +2112,7 @@ void FHttpNetworkReplayStreamer::HttpDownloadCheckpointFinished( FHttpRequestPtr
 {
 	UE_LOG( LogHttpReplay, Verbose, TEXT( "FHttpNetworkReplayStreamer::HttpDownloadCheckpointFinished." ) );
 
-	RequestFinished( EStreamerState::StreamingDown, EQueuedHttpRequestType::DownloadingCheckpoint, HttpRequest );
+	RequestFinished(EReplayStreamerState::Playback, EQueuedHttpRequestType::DownloadingCheckpoint, HttpRequest );
 
 	check( StreamArchive.IsLoading() );
 	check( GotoCheckpointDelegate.IsBound() );
@@ -2222,7 +2222,7 @@ void FHttpNetworkReplayStreamer::HttpRefreshViewerFinished( FHttpRequestPtr Http
 {
 	TSharedPtr< FQueuedHttpRequest > SavedFlightHttpRequest = InFlightHttpRequest;
 
-	RequestFinished( EStreamerState::StreamingDown, EQueuedHttpRequestType::RefreshingViewer, HttpRequest );
+	RequestFinished(EReplayStreamerState::Playback, EQueuedHttpRequestType::RefreshingViewer, HttpRequest );
 
 	if ( !bSucceeded || HttpResponse->GetResponseCode() != EHttpResponseCodes::NoContent )
 	{
@@ -2298,7 +2298,7 @@ void FHttpNetworkReplayStreamer::HttpEnumerateSessionsFinished( FHttpRequestPtr
 
 void FHttpNetworkReplayStreamer::HttpEnumerateCheckpointsFinished( FHttpRequestPtr HttpRequest, FHttpResponsePtr HttpResponse, bool bSucceeded )
 {
-	RequestFinished( EStreamerState::StreamingDown, EQueuedHttpRequestType::EnumeratingCheckpoints, HttpRequest );
+	RequestFinished(EReplayStreamerState::Playback, EQueuedHttpRequestType::EnumeratingCheckpoints, HttpRequest );
 
 	FStartStreamingResult Result;
 	Result.bRecording = false;
@@ -2374,7 +2374,7 @@ void FHttpNetworkReplayStreamer::HttpEnumerateEventsFinished(FHttpRequestPtr Htt
 
 void FHttpNetworkReplayStreamer::HttpAddUserFinished(FHttpRequestPtr HttpRequest, FHttpResponsePtr HttpResponse, bool bSucceeded)
 {
-	RequestFinished( EStreamerState::StreamingUp, EQueuedHttpRequestType::AddingUser, HttpRequest );
+	RequestFinished(EReplayStreamerState::Recording, EQueuedHttpRequestType::AddingUser, HttpRequest );
 
 	if ( bSucceeded && HttpResponse->GetResponseCode() == EHttpResponseCodes::NoContent )
 	{		
@@ -2424,7 +2424,7 @@ bool FHttpNetworkReplayStreamer::ProcessNextHttpRequest()
 		else if ( QueuedRequest->Type == EQueuedHttpRequestType::StopStreaming )
 		{
 			check( IsStreaming() );
-			StreamerState = EStreamerState::Idle;
+			StreamerState = EReplayStreamerState::Idle;
 			bStopStreamingCalled = false;
 			ensure( QueuedHttpRequests.Num() == 0 );
 			return ProcessNextHttpRequest();
@@ -2493,11 +2493,11 @@ void FHttpNetworkReplayStreamer::Tick( const float DeltaTime )
 		return;
 	}
 
-	if ( StreamerState == EStreamerState::StreamingUp )
+	if ( StreamerState == EReplayStreamerState::Recording )
 	{
 		ConditionallyFlushStream();
 	}
-	else if ( StreamerState == EStreamerState::StreamingDown )
+	else if ( StreamerState == EReplayStreamerState::Playback )
 	{
 		if ( IsTaskPendingOrInFlight( EQueuedHttpRequestType::StartDownloading ) )
 		{
@@ -2538,14 +2538,14 @@ bool FHttpNetworkReplayStreamer::HasPendingHttpRequests() const
 
 bool FHttpNetworkReplayStreamer::IsStreaming() const
 {
-	return StreamerState != EStreamerState::Idle;
+	return StreamerState != EReplayStreamerState::Idle;
 }
 
 void FHttpNetworkReplayStreamer::HttpDownloadCheckpointDeltaFinished( FHttpRequestPtr HttpRequest, FHttpResponsePtr HttpResponse, bool bSucceeded )
 {
 	UE_LOG( LogHttpReplay, Verbose, TEXT( "FHttpNetworkReplayStreamer::HttpDownloadCheckpointDeltaFinished." ) );
 
-	RequestFinished( EStreamerState::StreamingDown, EQueuedHttpRequestType::DownloadingCheckpoint, HttpRequest );
+	RequestFinished(EReplayStreamerState::Playback, EQueuedHttpRequestType::DownloadingCheckpoint, HttpRequest );
 
 	check( StreamArchive.IsLoading() );
 	check( GotoCheckpointDelegate.IsBound() );
diff --git a/Engine/Source/Runtime/NetworkReplayStreaming/HttpNetworkReplayStreaming/Public/HttpNetworkReplayStreaming.h b/Engine/Source/Runtime/NetworkReplayStreaming/HttpNetworkReplayStreaming/Public/HttpNetworkReplayStreaming.h
--- a/Engine/Source/Runtime/NetworkReplayStreaming/HttpNetworkReplayStreaming/Public/HttpNetworkReplayStreaming.h
+++ b/Engine/Source/Runtime/NetworkReplayStreaming/HttpNetworkReplayStreaming/Public/HttpNetworkReplayStreaming.h
@@ -246,6 +246,7 @@ public:
 
 	virtual ENetworkReplayError::Type GetLastError() const override;
 	virtual FString		GetReplayID() const override { return SessionName; }
+	virtual EReplayStreamerState GetReplayStreamerState() const override { return StreamerState; }
 	virtual void		SetTimeBufferHintSeconds(const float InTimeBufferHintSeconds) override {}
 	virtual void		RefreshHeader() override;
 	virtual void		DownloadHeader(const FDownloadHeaderCallback& Delegate);
@@ -298,16 +299,8 @@ public:
 	void InternalGotoTimeInMS(const uint32 TimeInMS, const FGotoCallback& Delegate, bool bDelta);
 	void InternalGotoCheckpointIndex(const int32 CheckpointIndex, const FGotoCallback& Delegate, const FHttpRequestCompleteDelegate& RequestDelegate);
 
-	/** EStreamerState - Overall state of the streamer */
-	enum class EStreamerState
-	{
-		Idle,						// The streamer is idle. Either we haven't started streaming yet, or we are done
-		StreamingUp,				// We are in the process of streaming a replay to the http server
-		StreamingDown				// We are in the process of streaming a replay from the http server
-	};
-
 	/** Delegates */
-	void RequestFinished( EStreamerState ExpectedStreamerState, EQueuedHttpRequestType::Type ExpectedType, FHttpRequestPtr HttpRequest );
+	void RequestFinished(EReplayStreamerState ExpectedStreamerState, EQueuedHttpRequestType::Type ExpectedType, FHttpRequestPtr HttpRequest );
 
 	void HttpStartDownloadingFinished( FHttpRequestPtr HttpRequest, FHttpResponsePtr HttpResponse, bool bSucceeded );
 	void HttpDownloadHeaderFinished( FHttpRequestPtr HttpRequest, FHttpResponsePtr HttpResponse, bool bSucceeded, FDownloadHeaderCallback Delegate);
@@ -347,7 +340,7 @@ public:
 	double					LastChunkTime;			// The last time we uploaded/downloaded a chunk
 	double					LastRefreshViewerTime;	// The last time we refreshed ourselves as an active viewer
 	double					LastRefreshCheckpointTime;
-	EStreamerState			StreamerState;			// Overall state of the streamer
+	EReplayStreamerState	StreamerState;			// Overall state of the streamer
 	bool					bStopStreamingCalled;
 	bool					bStreamIsLive;			// If true, we are viewing a live stream
 	FString					StreamMetadata;
diff --git a/Engine/Source/Runtime/NetworkReplayStreaming/InMemoryNetworkReplayStreaming/Private/InMemoryNetworkReplayStreaming.cpp b/Engine/Source/Runtime/NetworkReplayStreaming/InMemoryNetworkReplayStreaming/Private/InMemoryNetworkReplayStreaming.cpp
--- a/Engine/Source/Runtime/NetworkReplayStreaming/InMemoryNetworkReplayStreaming/Private/InMemoryNetworkReplayStreaming.cpp
+++ b/Engine/Source/Runtime/NetworkReplayStreaming/InMemoryNetworkReplayStreaming/Private/InMemoryNetworkReplayStreaming.cpp
@@ -55,7 +55,7 @@ void FInMemoryNetworkReplayStreamer::StartStreaming(const FStartStreamingParamet
 		FileAr->SetIsSaving(Params.bRecord);
 		FileAr->SetIsLoading(!Params.bRecord);
 		HeaderAr.Reset(new FMemoryReader(FoundReplay->Header));
-		StreamerState = EStreamerState::Playback;
+		StreamerState = EReplayStreamerState::Playback;
 	}
 	else
 	{
@@ -77,7 +77,7 @@ void FInMemoryNetworkReplayStreamer::StartStreaming(const FStartStreamingParamet
 
 		OwningFactory->Replays.Add(CurrentStreamName, MoveTemp(NewReplay));
 		
-		StreamerState = EStreamerState::Recording;
+		StreamerState = EReplayStreamerState::Recording;
 	}
 
 	// Notify immediately
@@ -91,7 +91,7 @@ void FInMemoryNetworkReplayStreamer::StartStreaming(const FStartStreamingParamet
 
 void FInMemoryNetworkReplayStreamer::StopStreaming()
 {
-	if (StreamerState == EStreamerState::Recording)
+	if (StreamerState == EReplayStreamerState::Recording)
 	{
 		FInMemoryReplay* FoundReplay = GetCurrentReplayChecked();
 
@@ -108,7 +108,7 @@ void FInMemoryNetworkReplayStreamer::StopStreaming()
 	FileAr.Reset();
 
 	CurrentStreamName.Empty();
-	StreamerState = EStreamerState::Idle;
+	StreamerState = EReplayStreamerState::Idle;
 }
 
 FArchive* FInMemoryNetworkReplayStreamer::GetHeaderArchive()
@@ -123,7 +123,7 @@ FArchive* FInMemoryNetworkReplayStreamer::GetStreamingArchive()
 
 void FInMemoryNetworkReplayStreamer::UpdateTotalDemoTime(uint32 TimeInMS)
 {
-	check(StreamerState == EStreamerState::Recording);
+	check(StreamerState == EReplayStreamerState::Recording);
 
 	FInMemoryReplay* FoundReplay = GetCurrentReplayChecked();
 
@@ -132,7 +132,7 @@ void FInMemoryNetworkReplayStreamer::UpdateTotalDemoTime(uint32 TimeInMS)
 
 uint32 FInMemoryNetworkReplayStreamer::GetTotalDemoTime() const
 {
-	check(StreamerState != EStreamerState::Idle);
+	check(StreamerState != EReplayStreamerState::Idle);
 
 	const FInMemoryReplay* FoundReplay = GetCurrentReplayChecked();
 
@@ -144,7 +144,7 @@ bool FInMemoryNetworkReplayStreamer::IsDataAvailable() const
 	// Assumptions:
 	// 1. All streamer instances run on the same thread, not simultaneously
 	// 2. A recording DemoNetDriver will write either no frames or entire frames each time it ticks
-	return StreamerState == EStreamerState::Playback && FileAr.IsValid() && FileAr->Tell() < FileAr->TotalSize();
+	return StreamerState == EReplayStreamerState::Playback && FileAr.IsValid() && FileAr->Tell() < FileAr->TotalSize();
 }
 
 bool FInMemoryNetworkReplayStreamer::IsLive() const
@@ -337,7 +337,7 @@ FArchive* FInMemoryNetworkReplayStreamer::GetCheckpointArchive()
 	// If the archive is null, and the API is being used properly, the caller is writing a checkpoint...
 	if ( CheckpointAr.Get() == nullptr )
 	{
-		check(StreamerState != EStreamerState::Playback);
+		check(StreamerState != EReplayStreamerState::Playback);
 
 		UE_LOG(LogMemoryReplay, Log, TEXT("FInMemoryNetworkReplayStreamer::GetCheckpointArchive. Creating new checkpoint."));
 
diff --git a/Engine/Source/Runtime/NetworkReplayStreaming/InMemoryNetworkReplayStreaming/Public/InMemoryNetworkReplayStreaming.h b/Engine/Source/Runtime/NetworkReplayStreaming/InMemoryNetworkReplayStreaming/Public/InMemoryNetworkReplayStreaming.h
--- a/Engine/Source/Runtime/NetworkReplayStreaming/InMemoryNetworkReplayStreaming/Public/InMemoryNetworkReplayStreaming.h
+++ b/Engine/Source/Runtime/NetworkReplayStreaming/InMemoryNetworkReplayStreaming/Public/InMemoryNetworkReplayStreaming.h
@@ -115,9 +115,9 @@ class FInMemoryNetworkReplayStreamer : public INetworkReplayStreamer, public FTi
 {
 public:
 	FInMemoryNetworkReplayStreamer(FInMemoryNetworkReplayStreamingFactory* InFactory) :
-		OwningFactory( InFactory ),
-		StreamerState( EStreamerState::Idle ),
-		TimeBufferHintSeconds( -1.0f )
+		OwningFactory(InFactory),
+		StreamerState(EReplayStreamerState::Idle),
+		TimeBufferHintSeconds(-1.0f)
 	{
 		check(OwningFactory != nullptr);
 	}
@@ -164,6 +164,7 @@ public:
 	virtual void RenameReplay(const FString& ReplayName, const FString& NewName, const FRenameReplayCallback& Delegate) override;
 	virtual void RenameReplay(const FString& ReplayName, const FString& NewName, const int32 UserIndex, const FRenameReplayCallback& Delegate) override;
 	virtual FString	GetReplayID() const override { return CurrentStreamName; }
+	virtual EReplayStreamerState GetReplayStreamerState() const override { return StreamerState; }
 	virtual void SetTimeBufferHintSeconds(const float InTimeBufferHintSeconds) override { TimeBufferHintSeconds = InTimeBufferHintSeconds; }
 	virtual void RefreshHeader() override {};
 	virtual void DownloadHeader(const FDownloadHeaderCallback& Delegate) override
@@ -242,16 +243,8 @@ private:
 	 */
 	FInMemoryReplay::FCheckpoint CheckpointCurrentlyBeingSaved;
 
-	/** EStreamerState - Overall state of the streamer */
-	enum class EStreamerState
-	{
-		Idle,					// The streamer is idle. Either we haven't started streaming yet, or we are done
-		Recording,				// We are in the process of recording a replay to disk
-		Playback,				// We are in the process of playing a replay from disk
-	};
-
 	/** Overall state of the streamer */
-	EStreamerState StreamerState;
+	EReplayStreamerState StreamerState;
 
 	/** Remember the name of the current stream, if any. */
 	FString CurrentStreamName;
diff --git a/Engine/Source/Runtime/NetworkReplayStreaming/LocalFileNetworkReplayStreaming/Private/LocalFileNetworkReplayStreaming.cpp b/Engine/Source/Runtime/NetworkReplayStreaming/LocalFileNetworkReplayStreaming/Private/LocalFileNetworkReplayStreaming.cpp
--- a/Engine/Source/Runtime/NetworkReplayStreaming/LocalFileNetworkReplayStreaming/Private/LocalFileNetworkReplayStreaming.cpp
+++ b/Engine/Source/Runtime/NetworkReplayStreaming/LocalFileNetworkReplayStreaming/Private/LocalFileNetworkReplayStreaming.cpp
@@ -146,7 +146,7 @@ FLocalFileNetworkReplayStreamer::FLocalFileNetworkReplayStreamer()
 	, bStopStreamingCalled(false)
 	, HighPriorityEndTime(0)
 	, LastGotoTimeInMS(-1)
-	, StreamerState(EStreamerState::Idle)
+	, StreamerState(EReplayStreamerState::Idle)
 	, StreamerLastError(ENetworkReplayError::None)
 	, DemoSavePath(GetDefaultDemoSavePath())
 	, bCacheFileReadsInMemory(false)
@@ -162,7 +162,7 @@ FLocalFileNetworkReplayStreamer::FLocalFileNetworkReplayStreamer(const FString&
 	, bStopStreamingCalled(false)
 	, HighPriorityEndTime(0)
 	, LastGotoTimeInMS(-1)
-	, StreamerState(EStreamerState::Idle)
+	, StreamerState(EReplayStreamerState::Idle)
 	, StreamerLastError(ENetworkReplayError::None)
 	, DemoSavePath(InDemoSavePath.EndsWith(TEXT("/")) ? InDemoSavePath : InDemoSavePath + FString("/"))
 	, bCacheFileReadsInMemory(false)
@@ -701,7 +701,7 @@ void FLocalFileNetworkReplayStreamer::StartStreaming(const FStartStreamingParame
 	if (!Params.bRecord)
 	{
 		// We are playing
-		StreamerState = EStreamerState::Playback;
+		StreamerState = EReplayStreamerState::Playback;
 
 		// Add the request to start loading
 		AddDelegateFileRequestToQueue<FStartStreamingResult>(EQueuedLocalFileRequestType::StartPlayback, 
@@ -760,7 +760,7 @@ void FLocalFileNetworkReplayStreamer::StartStreaming(const FStartStreamingParame
 	else
 	{
 		// We are recording
-		StreamerState = EStreamerState::Recording;
+		StreamerState = EReplayStreamerState::Recording;
 
 		uint64 TotalDiskFreeSpace = 0;
 		if (GetDemoFreeStorageSpace(TotalDiskFreeSpace, GetDemoPath()))
@@ -848,7 +848,7 @@ void FLocalFileNetworkReplayStreamer::CancelStreamingRequests()
 	// Empty the request queue
 	QueuedRequests.Empty();
 
-	StreamerState = EStreamerState::Idle;
+	StreamerState = EReplayStreamerState::Idle;
 	bStopStreamingCalled = false;
 }
 
@@ -888,7 +888,7 @@ void FLocalFileNetworkReplayStreamer::StopStreaming()
 
 	bStopStreamingCalled = true;
 
-	if (StreamerState == EStreamerState::Recording)
+	if (StreamerState == EReplayStreamerState::Recording)
 	{
 		// Flush any final pending stream
 		int32 TotalLengthInMS = CurrentReplayInfo.LengthInMS;
@@ -930,7 +930,7 @@ void FLocalFileNetworkReplayStreamer::StopStreaming()
 			StreamDataOffset = 0;
 			StreamChunkIndex = 0;
 			CurrentStreamName.Empty();
-			StreamerState = EStreamerState::Idle;
+			StreamerState = EReplayStreamerState::Idle;
 		});
 }
 
@@ -946,7 +946,7 @@ FArchive* FLocalFileNetworkReplayStreamer::GetStreamingArchive()
 
 void FLocalFileNetworkReplayStreamer::UpdateTotalDemoTime(uint32 TimeInMS)
 {
-	check(StreamerState == EStreamerState::Recording);
+	check(StreamerState == EReplayStreamerState::Recording);
 
 	CurrentReplayInfo.LengthInMS = TimeInMS;
 }
@@ -1103,7 +1103,7 @@ void FLocalFileNetworkReplayStreamer::AddUserToReplay(const FString& UserString)
 
 void FLocalFileNetworkReplayStreamer::AddEvent(const uint32 TimeInMS, const FString& Group, const FString& Meta, const TArray<uint8>& Data)
 {
-	if (StreamerState != EStreamerState::Recording)
+	if (StreamerState != EReplayStreamerState::Recording)
 	{
 		UE_LOG(LogLocalFileReplay, Warning, TEXT("FLocalFileNetworkReplayStreamer::AddEvent. Not recording."));
 		return;
@@ -1114,7 +1114,7 @@ void FLocalFileNetworkReplayStreamer::AddEvent(const uint32 TimeInMS, const FStr
 
 void FLocalFileNetworkReplayStreamer::AddOrUpdateEvent(const FString& Name, const uint32 TimeInMS, const FString& Group, const FString& Meta, const TArray<uint8>& Data)
 {
-	if (StreamerState != EStreamerState::Recording)
+	if (StreamerState != EReplayStreamerState::Recording)
 	{
 		UE_LOG(LogLocalFileReplay, Warning, TEXT("FLocalFileNetworkReplayStreamer::AddOrUpdateEvent. Not recording."));
 		return;
@@ -1803,7 +1803,7 @@ void FLocalFileNetworkReplayStreamer::FlushCheckpoint(const uint32 TimeInMS)
 
 void FLocalFileNetworkReplayStreamer::FlushCheckpointInternal(const uint32 TimeInMS)
 {
-	if (StreamerState != EStreamerState::Recording || CheckpointAr.Buffer.Num() == 0)
+	if (StreamerState != EReplayStreamerState::Recording || CheckpointAr.Buffer.Num() == 0)
 	{
 		// If there is no active session, or we are not recording, we don't need to flush
 		CheckpointAr.Buffer.Empty();
@@ -2495,11 +2495,11 @@ void FLocalFileNetworkReplayStreamer::Tick(float DeltaSeconds)
 		return;
 	}
 
-	if (StreamerState == EStreamerState::Recording)
+	if (StreamerState == EReplayStreamerState::Recording)
 	{
 		ConditionallyFlushStream();
 	}
-	else if (StreamerState == EStreamerState::Playback)
+	else if (StreamerState == EReplayStreamerState::Playback)
 	{
 		if (IsFileRequestPendingOrInProgress(EQueuedLocalFileRequestType::StartPlayback))
 		{
@@ -3025,7 +3025,7 @@ void FLocalFileNetworkReplayStreamer::OnFileRequestComplete(const TSharedPtr<FQu
 
 bool FLocalFileNetworkReplayStreamer::IsStreaming() const
 {
-	return (StreamerState != EStreamerState::Idle);
+	return (StreamerState != EReplayStreamerState::Idle);
 }
 
 void FLocalFileNetworkReplayStreamer::ConditionallyFlushStream()
diff --git a/Engine/Source/Runtime/NetworkReplayStreaming/LocalFileNetworkReplayStreaming/Public/LocalFileNetworkReplayStreaming.h b/Engine/Source/Runtime/NetworkReplayStreaming/LocalFileNetworkReplayStreaming/Public/LocalFileNetworkReplayStreaming.h
--- a/Engine/Source/Runtime/NetworkReplayStreaming/LocalFileNetworkReplayStreaming/Public/LocalFileNetworkReplayStreaming.h
+++ b/Engine/Source/Runtime/NetworkReplayStreaming/LocalFileNetworkReplayStreaming/Public/LocalFileNetworkReplayStreaming.h
@@ -456,6 +456,7 @@ public:
 	virtual void RenameReplay(const FString& ReplayName, const FString& NewName, const FRenameReplayCallback& Delegate) override;
 	virtual void RenameReplay(const FString& ReplayName, const FString& NewName, const int32 UserIndex, const FRenameReplayCallback& Delegate) override;
 	virtual FString	GetReplayID() const override { return CurrentStreamName; }
+	virtual EReplayStreamerState GetReplayStreamerState() const override { return StreamerState; }
 	virtual void SetTimeBufferHintSeconds(const float InTimeBufferHintSeconds) override {}
 	virtual void RefreshHeader() override;
 	virtual void DownloadHeader(const FDownloadHeaderCallback& Delegate) override;
@@ -643,16 +644,8 @@ protected:
 	/* Handle to the archive that will read/write checkpoint files */
 	FLocalFileStreamFArchive CheckpointAr;
 
-	/** EStreamerState - Overall state of the streamer */
-	enum class EStreamerState
-	{
-		Idle,					// The streamer is idle. Either we haven't started streaming yet, or we are done
-		Recording,				// We are in the process of recording a replay to disk
-		Playback,				// We are in the process of playing a replay from disk
-	};
-
 	/** Overall state of the streamer */
-	EStreamerState StreamerState;
+	EReplayStreamerState StreamerState;
 
 	ENetworkReplayError::Type StreamerLastError;
 
diff --git a/Engine/Source/Runtime/NetworkReplayStreaming/NetworkReplayStreaming/NetworkReplayStreaming.Build.cs b/Engine/Source/Runtime/NetworkReplayStreaming/NetworkReplayStreaming/NetworkReplayStreaming.Build.cs
--- a/Engine/Source/Runtime/NetworkReplayStreaming/NetworkReplayStreaming/NetworkReplayStreaming.Build.cs
+++ b/Engine/Source/Runtime/NetworkReplayStreaming/NetworkReplayStreaming/NetworkReplayStreaming.Build.cs
@@ -13,7 +13,7 @@ namespace UnrealBuildTool.Rules
 				{
 					"Core",
                     "Json",
-                    // ... add other public dependencies that you statically link with here ...
+					"Analytics",
 				}
 			);
 		}
diff --git a/Engine/Source/Runtime/NetworkReplayStreaming/NetworkReplayStreaming/Private/NetworkReplayStreaming.cpp b/Engine/Source/Runtime/NetworkReplayStreaming/NetworkReplayStreaming/Private/NetworkReplayStreaming.cpp
--- a/Engine/Source/Runtime/NetworkReplayStreaming/NetworkReplayStreaming/Private/NetworkReplayStreaming.cpp
+++ b/Engine/Source/Runtime/NetworkReplayStreaming/NetworkReplayStreaming/Private/NetworkReplayStreaming.cpp
@@ -7,6 +7,8 @@
 
 IMPLEMENT_MODULE( FNetworkReplayStreaming, NetworkReplayStreaming );
 
+FOnReplayGetAnalyticsAttributes INetworkReplayStreamer::OnReplayGetAnalyticsAttributes;
+
 INetworkReplayStreamingFactory& FNetworkReplayStreaming::GetFactory(const TCHAR* FactoryNameOverride)
 {
 	static const FString DefaultFactoryName = TEXT("LocalFileNetworkReplayStreaming");
@@ -127,4 +129,46 @@ bool FNetworkReplayStreaming::Exec(UWorld* InWorld, const TCHAR* Cmd, FOutputDev
 	}
 
 	return false;
+}
+
+FString LexToString(const EReplayStreamerState State)
+{
+	FString Str;
+
+	switch (State)
+	{
+	case EReplayStreamerState::Idle:
+		Str = TEXT("Idle");
+		break;
+	case EReplayStreamerState::Recording:
+		Str = TEXT("Recording");
+		break;
+	case EReplayStreamerState::Playback:
+		Str = TEXT("Playback");
+		break;
+	default:
+		Str = TEXT("Unknown");
+		break;
+	}
+
+	return Str;
+}
+
+TArray<FAnalyticsEventAttribute> INetworkReplayStreamer::AppendCommonReplayAttributes(TArray<FAnalyticsEventAttribute>&& Attrs) const
+{
+	static const FString Attrib_StreamerState = TEXT("StreamerState");
+	static const FString Attrib_ReplayName = TEXT("ReplayName");
+	static const FString Attrib_LengthInMS = TEXT("LengthInMS");
+
+	TArray<FAnalyticsEventAttribute> CommonAttributes = MoveTemp(Attrs);
+
+	AppendAnalyticsEventAttributeArray(CommonAttributes,
+		Attrib_StreamerState, LexToString(GetReplayStreamerState()),
+		Attrib_ReplayName, GetReplayID(),
+		Attrib_LengthInMS, GetTotalDemoTime()
+	);
+
+	INetworkReplayStreamer::OnReplayGetAnalyticsAttributes.Broadcast(this, CommonAttributes);
+
+	return CommonAttributes;
 }
\ No newline at end of file
diff --git a/Engine/Source/Runtime/NetworkReplayStreaming/NetworkReplayStreaming/Public/NetworkReplayStreaming.h b/Engine/Source/Runtime/NetworkReplayStreaming/NetworkReplayStreaming/Public/NetworkReplayStreaming.h
--- a/Engine/Source/Runtime/NetworkReplayStreaming/NetworkReplayStreaming/Public/NetworkReplayStreaming.h
+++ b/Engine/Source/Runtime/NetworkReplayStreaming/NetworkReplayStreaming/Public/NetworkReplayStreaming.h
@@ -9,6 +9,7 @@
 #include "Modules/ModuleInterface.h"
 #include "Modules/ModuleManager.h"
 #include "Serialization/JsonSerializerMacros.h"
+#include "AnalyticsEventAttribute.h"
 
 class IAnalyticsProvider;
 
@@ -473,6 +474,17 @@ struct FStartStreamingParameters
 	bool bRecord;
 };
 
+enum class EReplayStreamerState : uint8
+{
+	Idle,					// The streamer is idle. Either we haven't started activity yet, or we are done
+	Recording,				// We are in the process of recording a replay
+	Playback,				// We are in the process of playing a replay
+};
+
+FString NETWORKREPLAYSTREAMING_API LexToString(const EReplayStreamerState State);
+
+DECLARE_MULTICAST_DELEGATE_TwoParams(FOnReplayGetAnalyticsAttributes, const class INetworkReplayStreamer* /*Streamer*/, TArray<FAnalyticsEventAttribute>& /*Attributes*/);
+
 /**
  * Generic interface for network replay streaming
  *
@@ -502,7 +514,9 @@ public:
 	virtual void UpdateTotalDemoTime(uint32 TimeInMS) = 0;
 	virtual void UpdatePlaybackTime(uint32 TimeInMS) = 0;
 
+	/** Time in milliseconds */
 	virtual uint32 GetTotalDemoTime() const = 0;
+
 	virtual bool IsDataAvailable() const = 0;
 	virtual void SetHighPriorityTimeRange(const uint32 StartTimeInMS, const uint32 EndTimeInMS) = 0;
 	virtual bool IsDataAvailableForTimeRange(const uint32 StartTimeInMS, const uint32 EndTimeInMS) = 0;
@@ -563,6 +577,9 @@ public:
 	/** Returns the active replay name */
 	virtual FString	GetReplayID() const = 0;
 
+	/** Return current recording/playback state */
+	virtual EReplayStreamerState GetReplayStreamerState() const { return EReplayStreamerState::Idle; };
+
 	/**
 	 * Attempts to delete the stream with the specified name. May execute asynchronously.
 	 *
@@ -621,7 +638,12 @@ public:
 
 	virtual void SetAnalyticsProvider(TSharedPtr<IAnalyticsProvider>& InProvider) {}
 
+	virtual TArray<FAnalyticsEventAttribute> AppendCommonReplayAttributes(TArray<FAnalyticsEventAttribute>&& Attrs) const;
+
 	virtual void Exec(const TCHAR* Cmd, FOutputDevice& Ar) {}
+
+	/** Called from base streamer interface AppendCommonReplayAttributes to set common attributes */
+	static FOnReplayGetAnalyticsAttributes OnReplayGetAnalyticsAttributes;
 };
 
 /** Replay streamer factory */
diff --git a/Engine/Source/Runtime/NetworkReplayStreaming/NullNetworkReplayStreaming/Private/NullNetworkReplayStreaming.cpp b/Engine/Source/Runtime/NetworkReplayStreaming/NullNetworkReplayStreaming/Private/NullNetworkReplayStreaming.cpp
--- a/Engine/Source/Runtime/NetworkReplayStreaming/NullNetworkReplayStreaming/Private/NullNetworkReplayStreaming.cpp
+++ b/Engine/Source/Runtime/NetworkReplayStreaming/NullNetworkReplayStreaming/Private/NullNetworkReplayStreaming.cpp
@@ -216,7 +216,7 @@ void FNullNetworkReplayStreamer::StartStreaming(const FStartStreamingParameters&
 		// Open file for reading
 		ReopenStreamFileForReading();
 		HeaderAr.Reset( IFileManager::Get().CreateFileReader( *FullHeaderFilename, FILEREAD_AllowWrite ) );
-		StreamerState = EStreamerState::Playback;
+		StreamerState = EReplayStreamerState::Playback;
 	}
 	else
 	{
@@ -229,7 +229,7 @@ void FNullNetworkReplayStreamer::StartStreaming(const FStartStreamingParameters&
 		// Open file for writing
 		FileAr.Reset( IFileManager::Get().CreateFileWriter( *FullDemoFilename, FILEWRITE_AllowRead ) );
 		HeaderAr.Reset( IFileManager::Get().CreateFileWriter( *FullHeaderFilename, FILEWRITE_AllowRead ) );
-		StreamerState = EStreamerState::Recording;
+		StreamerState = EReplayStreamerState::Recording;
 
 		CurrentCheckpointIndex = 0;
 
@@ -251,7 +251,7 @@ void FNullNetworkReplayStreamer::StartStreaming(const FStartStreamingParameters&
 
 void FNullNetworkReplayStreamer::StopStreaming()
 {
-	if (StreamerState == EStreamerState::Recording)
+	if (StreamerState == EReplayStreamerState::Recording)
 	{
 		WriteReplayInfo(CurrentStreamName, ReplayInfo);
 	}
@@ -262,7 +262,7 @@ void FNullNetworkReplayStreamer::StopStreaming()
 	FileAr.Reset();
 
 	CurrentStreamName.Empty();
-	StreamerState = EStreamerState::Idle;
+	StreamerState = EReplayStreamerState::Idle;
 }
 
 FArchive* FNullNetworkReplayStreamer::GetHeaderArchive()
@@ -277,14 +277,14 @@ FArchive* FNullNetworkReplayStreamer::GetStreamingArchive()
 
 void FNullNetworkReplayStreamer::UpdateTotalDemoTime(uint32 TimeInMS)
 {
-	check(StreamerState == EStreamerState::Recording);
+	check(StreamerState == EReplayStreamerState::Recording);
 
 	ReplayInfo.LengthInMS = TimeInMS;
 }
 
 bool FNullNetworkReplayStreamer::IsDataAvailable() const
 {
-	check(StreamerState == EStreamerState::Playback);
+	check(StreamerState == EReplayStreamerState::Playback);
 
 	return FileAr.IsValid() && FileAr->Tell() < ReplayInfo.DemoFileLastOffset;
 }
@@ -517,7 +517,7 @@ FArchive* FNullNetworkReplayStreamer::GetCheckpointArchive()
 	if ( CheckpointAr.Get() == nullptr )
 	{
 		// Create a file writer for the next checkpoint index.
-		check(StreamerState != EStreamerState::Playback);
+		check(StreamerState != EReplayStreamerState::Playback);
 
 		FString NextCheckpointFileName = FString::Printf( TEXT( "checkpoint%d" ), CurrentCheckpointIndex );
 
@@ -712,7 +712,7 @@ void FNullNetworkReplayStreamer::Tick(float DeltaSeconds)
 {
 	// This relies on the fact that the DemoNetDriver isn't currently in the middle of its own tick,
 	// and has either read or written a whole demo frame.
-	if (StreamerState == EStreamerState::Playback)
+	if (StreamerState == EReplayStreamerState::Playback)
 	{
 		// Re-read replay info
 		UpdateReplayInfoIfValid();
@@ -728,7 +728,7 @@ void FNullNetworkReplayStreamer::Tick(float DeltaSeconds)
 			}
 		}
 	}
-	else if (StreamerState == EStreamerState::Recording)
+	else if (StreamerState == EReplayStreamerState::Recording)
 	{
 		// Note the size of the file between demo frames
 		if (FileAr.IsValid() && ReplayInfo.DemoFileLastOffset < FileAr->Tell())
diff --git a/Engine/Source/Runtime/NetworkReplayStreaming/NullNetworkReplayStreaming/Public/NullNetworkReplayStreaming.h b/Engine/Source/Runtime/NetworkReplayStreaming/NullNetworkReplayStreaming/Public/NullNetworkReplayStreaming.h
--- a/Engine/Source/Runtime/NetworkReplayStreaming/NullNetworkReplayStreaming/Public/NullNetworkReplayStreaming.h
+++ b/Engine/Source/Runtime/NetworkReplayStreaming/NullNetworkReplayStreaming/Public/NullNetworkReplayStreaming.h
@@ -38,7 +38,7 @@ class FNullNetworkReplayStreamer : public INetworkReplayStreamer, public FTickab
 {
 public:
 	FNullNetworkReplayStreamer() :
-		StreamerState( EStreamerState::Idle ),
+		StreamerState(EReplayStreamerState::Idle),
 		CurrentCheckpointIndex( 0 ),
 		LastKnownFileSize( 0 )
 	{}
@@ -87,6 +87,7 @@ public:
 	virtual void RenameReplay(const FString& ReplayName, const FString& NewName, const FRenameReplayCallback& Delegate) override;
 	virtual void RenameReplay(const FString& ReplayName, const FString& NewName, const int32 UserIndex, const FRenameReplayCallback& Delegate) override;
 	virtual FString	GetReplayID() const override { return CurrentStreamName; }
+	virtual EReplayStreamerState GetReplayStreamerState() const override { return StreamerState; }
 	virtual void SetTimeBufferHintSeconds(const float InTimeBufferHintSeconds) override {}
 	virtual void RefreshHeader() override {};
 	virtual void DownloadHeader(const FDownloadHeaderCallback& Delegate) override
@@ -136,16 +137,8 @@ private:
 	/* Handle to the archive that will read/write checkpoint files */
 	TUniquePtr<FArchive> CheckpointAr;
 
-	/** EStreamerState - Overall state of the streamer */
-	enum class EStreamerState
-	{
-		Idle,					// The streamer is idle. Either we haven't started streaming yet, or we are done
-		Recording,				// We are in the process of recording a replay to disk
-		Playback,				// We are in the process of playing a replay from disk
-	};
-
 	/** Overall state of the streamer */
-	EStreamerState StreamerState;
+	EReplayStreamerState StreamerState;
 
 	/** Remember the name of the current stream, if any. */
 	FString CurrentStreamName;
