Add overload of AddRequestToCache that is move aware
Trim memory from stream archive before appending new chunk to avoid realloc when possible
Adding Reset method to FLocalFileStreamFArchive
Empty -> Reset on internal archive buffers that are frequently reused

diff --git a/Engine/Source/Runtime/NetworkReplayStreaming/LocalFileNetworkReplayStreaming/Private/LocalFileNetworkReplayStreaming.cpp b/Engine/Source/Runtime/NetworkReplayStreaming/LocalFileNetworkReplayStreaming/Private/LocalFileNetworkReplayStreaming.cpp
--- a/Engine/Source/Runtime/NetworkReplayStreaming/LocalFileNetworkReplayStreaming/Private/LocalFileNetworkReplayStreaming.cpp
+++ b/Engine/Source/Runtime/NetworkReplayStreaming/LocalFileNetworkReplayStreaming/Private/LocalFileNetworkReplayStreaming.cpp
@@ -72,6 +72,7 @@ namespace LocalFileReplay
 #endif
 
 	TAutoConsoleVariable<int32> CVarReplayRecordingMinSpace(TEXT("localReplay.ReplayRecordingMinSpace"), 20 * (1024 * 1024), TEXT("Minimum space needed to start recording a replay."));
+	TAutoConsoleVariable<float> CVarMinLoadNextChunkDelaySeconds(TEXT("localReplay.MinLoadNextChunkDelaySeconds"), 3.0f, TEXT("Minimum time to wait between conditional chunk loads."));
 };
 
 const uint32 FLocalFileNetworkReplayStreamer::FileMagic = 0x1CA2E27F;
@@ -925,8 +926,7 @@ void FLocalFileNetworkReplayStreamer::StopStreaming()
 			bStopStreamingCalled = false;
 			StreamAr.SetIsLoading(false);
 			StreamAr.SetIsSaving(false);
-			StreamAr.Buffer.Empty();
-			StreamAr.Pos = 0;
+			StreamAr.Reset();
 			StreamDataOffset = 0;
 			StreamChunkIndex = 0;
 			CurrentStreamName.Empty();
@@ -1775,8 +1775,7 @@ void FLocalFileNetworkReplayStreamer::FlushStream(const uint32 TimeInMS)
 			}
 		});
 
-	StreamAr.Buffer.Empty();
-	StreamAr.Pos = 0;
+	StreamAr.Reset();
 
 	// Keep track of the time range we have in our buffer, so we can accurately upload that each time we submit a chunk
 	StreamTimeRange.Min = StreamTimeRange.Max;
@@ -1928,7 +1927,7 @@ void FLocalFileNetworkReplayStreamer::FlushCheckpointInternal(const uint32 TimeI
 			}
 		});
 
-	CheckpointAr.Buffer.Empty();
+	CheckpointAr.Buffer.Reset();
 	CheckpointAr.Pos = 0;	
 }
 
@@ -1952,13 +1951,12 @@ void FLocalFileNetworkReplayStreamer::GotoCheckpointIndex(const int32 Checkpoint
 			[this, Delegate]()
 			{
 				// Make sure to reset the checkpoint archive (this is how we signify that the engine should start from the beginning of the stream (we don't need a checkpoint for that))
-				CheckpointAr.Buffer.Empty();
-				CheckpointAr.Pos = 0;
+				CheckpointAr.Reset();
 
 				if (!IsDataAvailableForTimeRange(0, LastGotoTimeInMS))
 				{
 					// Completely reset our stream (we're going to start loading from the start of the checkpoint)
-					StreamAr.Buffer.Empty();
+					StreamAr.Buffer.Reset();
 
 					StreamDataOffset = 0;
 
@@ -2138,9 +2136,7 @@ void FLocalFileNetworkReplayStreamer::GotoCheckpointIndex(const int32 Checkpoint
 			if (!bIsDataAvailableForTimeRange)
 			{
 				// Completely reset our stream (we're going to start loading from the start of the checkpoint)
-				StreamAr.Buffer.Empty();
-				StreamAr.Pos = 0;
-				StreamAr.bAtEndOfReplay = false;
+				StreamAr.Reset();
 
 				// Reset any time we were waiting on in the past
 				HighPriorityEndTime = 0;
@@ -2311,9 +2307,7 @@ void FLocalFileNetworkReplayStreamer::GotoCheckpointIndex(const int32 Checkpoint
 				if (!bIsDataAvailableForTimeRange)
 				{
 					// Completely reset our stream (we're going to start loading from the start of the checkpoint)
-					StreamAr.Buffer.Empty();
-					StreamAr.Pos = 0;
-					StreamAr.bAtEndOfReplay = false;
+					StreamAr.Reset();
 
 					// Reset any time we were waiting on in the past
 					HighPriorityEndTime = 0;
@@ -2973,8 +2967,7 @@ void FLocalFileNetworkReplayStreamer::WriteHeader()
 		});
 
 	// We're done with the header archive
-	HeaderAr.Buffer.Empty();
-	HeaderAr.Pos = 0;
+	HeaderAr.Reset();
 
 	LastChunkTime = FPlatformTime::Seconds();
 }
@@ -3071,30 +3064,29 @@ void FLocalFileNetworkReplayStreamer::ConditionallyLoadNextChunk()
 	// If it's not critical to load the next chunk (i.e. we're not scrubbing or at the end already), then check to see if we should grab the next chunk
 	if (!bReallyNeedToLoadChunk)
 	{
-		const double MIN_WAIT_FOR_NEXT_CHUNK_IN_SECONDS = 3;
-
 		const double LoadElapsedTime = FPlatformTime::Seconds() - LastChunkTime;
 
-		if (LoadElapsedTime < MIN_WAIT_FOR_NEXT_CHUNK_IN_SECONDS)
+		// Unless it's critical (i.e. bReallyNeedToLoadChunk is true), never try faster than the min delay
+		if (LoadElapsedTime < LocalFileReplay::CVarMinLoadNextChunkDelaySeconds.GetValueOnAnyThread())
 		{
-			return;		// Unless it's critical (i.e. bReallyNeedToLoadChunk is true), never try faster than MIN_WAIT_FOR_NEXT_CHUNK_IN_SECONDS
+			return;		
 		}
 
 		if ((StreamTimeRange.Max > StreamTimeRange.Min) && (StreamAr.Buffer.Num() > 0))
 		{
 			// Make a guess on how far we're in
-			const float PercentIn		= StreamAr.Buffer.Num() > 0 ? ( float )StreamAr.Pos / ( float )StreamAr.Buffer.Num() : 0.0f;
-			const float TotalStreamTime	= ( float )( StreamTimeRange.Size() ) / 1000.0f;
-			const float CurrentTime		= TotalStreamTime * PercentIn;
-			const float TimeLeft		= TotalStreamTime - CurrentTime;
+			const float PercentIn		= StreamAr.Buffer.Num() > 0 ? (float)StreamAr.Pos / (float)StreamAr.Buffer.Num() : 0.0f;
+			const float TotalStreamTimeSeconds = (float)(StreamTimeRange.Size()) / 1000.0f;
+			const float CurrentTime		= TotalStreamTimeSeconds * PercentIn;
+			const float TimeLeft		= TotalStreamTimeSeconds - CurrentTime;
 
 			// Determine if we have enough buffer to stop streaming for now
-			const float MAX_BUFFERED_TIME = LocalFileReplay::CVarChunkUploadDelayInSeconds.GetValueOnAnyThread() * 0.5f;
+			const float MaxBufferedTimeSeconds = LocalFileReplay::CVarChunkUploadDelayInSeconds.GetValueOnAnyThread() * 0.5f;
 
-			if (TimeLeft > MAX_BUFFERED_TIME)
+			if (TimeLeft > MaxBufferedTimeSeconds)
 			{
-				// Don't stream ahead by more than MAX_BUFFERED_TIME seconds
-				UE_LOG(LogLocalFileReplay, VeryVerbose, TEXT("ConditionallyLoadNextChunk. Cancelling due buffer being large enough. TotalStreamTime: %2.2f, PercentIn: %2.2f, TimeLeft: %2.2f"), TotalStreamTime, PercentIn, TimeLeft);
+				// Don't stream ahead by more than MaxBufferedTimeSeconds seconds
+				UE_LOG(LogLocalFileReplay, VeryVerbose, TEXT("ConditionallyLoadNextChunk. Cancelling due buffer being large enough. TotalStreamTime: %2.2f, PercentIn: %2.2f, TimeLeft: %2.2f"), TotalStreamTimeSeconds, PercentIn, TimeLeft);
 				return;
 			}
 		}
@@ -3201,8 +3193,7 @@ void FLocalFileNetworkReplayStreamer::ConditionallyLoadNextChunk()
 			// Make sure our stream chunk index didn't change under our feet
 			if (RequestedStreamChunkIndex != StreamChunkIndex)
 			{
-				StreamAr.Buffer.Empty();
-				StreamAr.Pos = 0;
+				StreamAr.Reset();
 				SetLastError(ENetworkReplayError::ServiceUnavailable);
 				return;
 			}
@@ -3219,10 +3210,7 @@ void FLocalFileNetworkReplayStreamer::ConditionallyLoadNextChunk()
 
 				check(StreamTimeRange.IsValid());
 
-				AddRequestToCache(CurrentReplayInfo.DataChunks[RequestedStreamChunkIndex].ChunkIndex, RequestData.DataBuffer);
-
-				StreamAr.Buffer.Append(RequestData.DataBuffer);
-
+				// make space before appending
 				int32 MaxBufferedChunks = LocalFileReplay::CVarMaxBufferedStreamChunks.GetValueOnAnyThread();
 				if (MaxBufferedChunks > 0)
 				{
@@ -3235,7 +3223,8 @@ void FLocalFileNetworkReplayStreamer::ConditionallyLoadNextChunk()
 							// can't remove chunks if we're actively seeking within that data
 							if (StreamAr.Pos >= TrimBytes)
 							{
-								StreamAr.Buffer.RemoveAt(0, TrimBytes);
+								// don't realloc, we're about to append anyway
+								StreamAr.Buffer.RemoveAt(0, TrimBytes, false);
 								StreamAr.Pos -= TrimBytes;
 
 								StreamTimeRange.Min = CurrentReplayInfo.DataChunks[MinChunkIndex].Time1;
@@ -3245,7 +3234,11 @@ void FLocalFileNetworkReplayStreamer::ConditionallyLoadNextChunk()
 							}
 						}
 					}
-				}				
+				}
+
+				StreamAr.Buffer.Append(RequestData.DataBuffer);
+
+				AddRequestToCache(CurrentReplayInfo.DataChunks[RequestedStreamChunkIndex].ChunkIndex, MoveTemp(RequestData.DataBuffer));
 
 				StreamChunkIndex++;
 			}
@@ -3294,6 +3287,13 @@ void FLocalFileNetworkReplayStreamer::ConditionallyRefreshReplayInfo()
 }
 
 void FLocalFileNetworkReplayStreamer::AddRequestToCache(int32 ChunkIndex, const TArray<uint8>& RequestData)
+{
+	LLM_SCOPE(ELLMTag::Replays);
+	TArray<uint8> DataCopy = RequestData;
+	AddRequestToCache(ChunkIndex, MoveTemp(DataCopy));
+}
+
+void FLocalFileNetworkReplayStreamer::AddRequestToCache(int32 ChunkIndex, TArray<uint8>&& RequestData)
 {
 	if (!CurrentReplayInfo.bIsValid)
 	{
@@ -3312,7 +3312,7 @@ void FLocalFileNetworkReplayStreamer::AddRequestToCache(int32 ChunkIndex, const
 
 	// Add to cache (or freshen existing entry)
 	LLM_SCOPE(ELLMTag::Replays);
-	RequestCache.Add(ChunkIndex, MakeShareable(new FCachedFileRequest(RequestData, FPlatformTime::Seconds())));
+	RequestCache.Add(ChunkIndex, MakeShareable(new FCachedFileRequest(MoveTemp(RequestData), FPlatformTime::Seconds())));
 
 	// Anytime we add something to cache, make sure it's within budget
 	CleanupRequestCache();
diff --git a/Engine/Source/Runtime/NetworkReplayStreaming/LocalFileNetworkReplayStreaming/Public/LocalFileNetworkReplayStreaming.h b/Engine/Source/Runtime/NetworkReplayStreaming/LocalFileNetworkReplayStreaming/Public/LocalFileNetworkReplayStreaming.h
--- a/Engine/Source/Runtime/NetworkReplayStreaming/LocalFileNetworkReplayStreaming/Public/LocalFileNetworkReplayStreaming.h
+++ b/Engine/Source/Runtime/NetworkReplayStreaming/LocalFileNetworkReplayStreaming/Public/LocalFileNetworkReplayStreaming.h
@@ -143,6 +143,13 @@ public:
 	virtual void	Seek(int64 InPos) override;
 	virtual bool	AtEnd() override;
 
+	void Reset()
+	{
+		Buffer.Reset();
+		Pos = 0;
+		bAtEndOfReplay = false;
+	}
+
 	TArray<uint8>	Buffer;
 	int32			Pos;
 	bool			bAtEndOfReplay;
@@ -233,6 +240,12 @@ public:
 	{
 	}
 
+	FCachedFileRequest(TArray<uint8>&& InRequestData, const double InLastAccessTime)
+		: RequestData(MoveTemp(InRequestData))
+		, LastAccessTime(InLastAccessTime)
+	{
+	}
+
 	TArray<uint8> RequestData;
 	double LastAccessTime;
 };
@@ -655,6 +668,7 @@ protected:
 	FString DemoSavePath;
 
 	void AddRequestToCache(int32 ChunkIndex, const TArray<uint8>& RequestData);
+	void AddRequestToCache(int32 ChunkIndex, TArray<uint8>&& RequestData);
 	void CleanupRequestCache();
 
 	bool bCacheFileReadsInMemory;
