Adding additional replay result codes for decryption failure and event not found

diff --git a/Engine/Source/Runtime/NetworkReplayStreaming/LocalFileNetworkReplayStreaming/Private/LocalFileNetworkReplayStreaming.cpp b/Engine/Source/Runtime/NetworkReplayStreaming/LocalFileNetworkReplayStreaming/Private/LocalFileNetworkReplayStreaming.cpp
--- a/Engine/Source/Runtime/NetworkReplayStreaming/LocalFileNetworkReplayStreaming/Private/LocalFileNetworkReplayStreaming.cpp
+++ b/Engine/Source/Runtime/NetworkReplayStreaming/LocalFileNetworkReplayStreaming/Private/LocalFileNetworkReplayStreaming.cpp
@@ -1357,10 +1357,14 @@ void FLocalFileNetworkReplayStreamer::RequestEventData_Internal(const FString& R
 					TSharedPtr<FArchive> LocalFileAr = CreateLocalFileReader(FullDemoFilename);
 					if (LocalFileAr.IsValid())
 					{
+						bool bEventFound = false;
+
 						for (const FLocalFileEventInfo& EventInfo : StoredReplayInfo.Events)
 						{
 							if (EventInfo.Id == EventID)
 							{
+								bEventFound = true;
+
 								LocalFileAr->Seek(EventInfo.EventDataOffset);
 
 								TArray<uint8> EventData;
@@ -1379,7 +1383,7 @@ void FLocalFileNetworkReplayStreamer::RequestEventData_Internal(const FString& R
 										if (!DecryptBuffer(EventData, PlaintextData, StoredReplayInfo.EncryptionKey))
 										{
 											UE_LOG(LogLocalFileReplay, Error, TEXT("FLocalFileNetworkReplayStreamer::RequestEventData_Internal. DecryptBuffer failed."));
-											RequestData.DelegateResult.Result = EStreamingOperationResult::Unspecified;
+											RequestData.DelegateResult.Result = EStreamingOperationResult::DecryptFailure;
 											break;
 										}
 
@@ -1388,7 +1392,7 @@ void FLocalFileNetworkReplayStreamer::RequestEventData_Internal(const FString& R
 									else
 									{
 										UE_LOG(LogLocalFileReplay, Error, TEXT("FLocalFileNetworkReplayStreamer::RequestEventData_Internal. Encrypted event but streamer does not support encryption."));
-										RequestData.DelegateResult.Result = EStreamingOperationResult::Unspecified;
+										RequestData.DelegateResult.Result = EStreamingOperationResult::Unsupported;
 										break;
 									}
 								}
@@ -1401,7 +1405,23 @@ void FLocalFileNetworkReplayStreamer::RequestEventData_Internal(const FString& R
 						}
 
 						LocalFileAr = nullptr;
+
+						// we didn't find the event
+						if (!bEventFound)
+						{
+							RequestData.DelegateResult.Result = EStreamingOperationResult::EventNotFound;
+						}
 					}
+					else
+					{
+						UE_LOG(LogLocalFileReplay, Error, TEXT("FLocalFileNetworkReplayStreamer::RequestEventData_Internal. Unable to read replay file: %s"), *FullDemoFilename);
+						RequestData.DelegateResult.Result = EStreamingOperationResult::Unspecified;
+					}
+				}
+				else
+				{
+					UE_LOG(LogLocalFileReplay, Error, TEXT("FLocalFileNetworkReplayStreamer::RequestEventData_Internal. Failed to read the replay info: %s"), *FileName);
+					RequestData.DelegateResult.Result = EStreamingOperationResult::Unspecified;
 				}
 			}
 		});
@@ -1444,10 +1464,14 @@ void FLocalFileNetworkReplayStreamer::RequestEventGroupData(const FString& Repla
 					TSharedPtr<FArchive> LocalFileAr = CreateLocalFileReader(FullDemoFilename);
 					if (LocalFileAr.IsValid())
 					{
+						bool bGroupFound = false;
+
 						for (const FLocalFileEventInfo& EventInfo : StoredReplayInfo.Events)
 						{
 							if (EventInfo.Group == Group)
 							{
+								bGroupFound = true;
+
 								LocalFileAr->Seek(EventInfo.EventDataOffset);
 
 								RequestData.DelegateResult.Result = EStreamingOperationResult::Success;
@@ -1475,7 +1499,7 @@ void FLocalFileNetworkReplayStreamer::RequestEventGroupData(const FString& Repla
 										if (!DecryptBuffer(EventData, PlaintextData, StoredReplayInfo.EncryptionKey))
 										{
 											UE_LOG(LogLocalFileReplay, Error, TEXT("FLocalFileNetworkReplayStreamer::RequestEventData_Internal. DecryptBuffer failed."));
-											RequestData.DelegateResult.Result = EStreamingOperationResult::Unspecified;
+											RequestData.DelegateResult.Result = EStreamingOperationResult::DecryptFailure;
 											break;
 										}
 
@@ -1484,7 +1508,7 @@ void FLocalFileNetworkReplayStreamer::RequestEventGroupData(const FString& Repla
 									else
 									{
 										UE_LOG(LogLocalFileReplay, Error, TEXT("FLocalFileNetworkReplayStreamer::RequestEventData_Internal. Encrypted event but streamer does not support encryption."));
-										RequestData.DelegateResult.Result = EStreamingOperationResult::Unspecified;
+										RequestData.DelegateResult.Result = EStreamingOperationResult::Unsupported;
 										break;
 									}
 								}
@@ -1495,8 +1519,24 @@ void FLocalFileNetworkReplayStreamer::RequestEventGroupData(const FString& Repla
 							}
 						}
 
+						// we didn't find the group
+						if (!bGroupFound)
+						{
+							RequestData.DelegateResult.Result = EStreamingOperationResult::EventNotFound;
+						}
+
 						LocalFileAr = nullptr;
 					}
+					else
+					{
+						UE_LOG(LogLocalFileReplay, Error, TEXT("FLocalFileNetworkReplayStreamer::RequestEventData_Internal. Failed to read the replay file: %s"), *FullDemoFilename);
+						RequestData.DelegateResult.Result = EStreamingOperationResult::Unspecified;
+					}
+				}
+				else
+				{
+					UE_LOG(LogLocalFileReplay, Error, TEXT("FLocalFileNetworkReplayStreamer::RequestEventData_Internal. Failed to read the replay info: %s"), *FileName);
+					RequestData.DelegateResult.Result = EStreamingOperationResult::Unspecified;
 				}
 			}
 		});
diff --git a/Engine/Source/Runtime/NetworkReplayStreaming/NetworkReplayStreaming/Public/NetworkReplayStreaming.h b/Engine/Source/Runtime/NetworkReplayStreaming/NetworkReplayStreaming/Public/NetworkReplayStreaming.h
--- a/Engine/Source/Runtime/NetworkReplayStreaming/NetworkReplayStreaming/Public/NetworkReplayStreaming.h
+++ b/Engine/Source/Runtime/NetworkReplayStreaming/NetworkReplayStreaming/Public/NetworkReplayStreaming.h
@@ -185,6 +185,8 @@ enum class EStreamingOperationResult
 	NotEnoughSlots,	//! The operation failed due to reaching a predefined replay limit.
 	Unspecified,	//! The operation failed for unspecified reasons.
 	UnfinishedTask, //! The operation failed due to an outstanding task.
+	EventNotFound,	//! The operation failed because the event could not be found.
+	DecryptFailure,	//! The operation failed because decryption failed.
 };
 
 /**
diff --git a/Engine/Source/Runtime/NetworkReplayStreaming/SaveGameNetworkReplayStreaming/Private/SaveGameNetworkReplayStreaming.cpp b/Engine/Source/Runtime/NetworkReplayStreaming/SaveGameNetworkReplayStreaming/Private/SaveGameNetworkReplayStreaming.cpp
--- a/Engine/Source/Runtime/NetworkReplayStreaming/SaveGameNetworkReplayStreaming/Private/SaveGameNetworkReplayStreaming.cpp
+++ b/Engine/Source/Runtime/NetworkReplayStreaming/SaveGameNetworkReplayStreaming/Private/SaveGameNetworkReplayStreaming.cpp
@@ -1453,6 +1453,12 @@ void FSaveGameNetworkReplayStreamer::RequestEventData_Internal(const FString& Re
 				break;
 			}
 		}
+
+		// we didn't find the event
+		if (Result.Result != EStreamingOperationResult::Success)
+		{
+			Result.Result = EStreamingOperationResult::EventNotFound;
+		}
 	}
 }
 
@@ -1548,6 +1554,12 @@ void FSaveGameNetworkReplayStreamer::RequestEventGroupData_Internal(const FStrin
 				EventDataResult.ReplayEventListItem = MoveTemp(MetaData.VersionedInfo.EventData[i]);
 			}
 		}
+
+		// we didn't find the event group
+		if (Result.Result != EStreamingOperationResult::Success)
+		{
+			Result.Result = EStreamingOperationResult::EventNotFound;
+		}
 	}
 }
 
