Adding max event/size cvars for sanity checking save game replays

diff --git a/Engine/Source/Runtime/NetworkReplayStreaming/SaveGameNetworkReplayStreaming/Private/SaveGameNetworkReplayStreaming.cpp b/Engine/Source/Runtime/NetworkReplayStreaming/SaveGameNetworkReplayStreaming/Private/SaveGameNetworkReplayStreaming.cpp
--- a/Engine/Source/Runtime/NetworkReplayStreaming/SaveGameNetworkReplayStreaming/Private/SaveGameNetworkReplayStreaming.cpp
+++ b/Engine/Source/Runtime/NetworkReplayStreaming/SaveGameNetworkReplayStreaming/Private/SaveGameNetworkReplayStreaming.cpp
@@ -600,6 +600,12 @@ namespace SaveGameReplay
 		TOptional<FExtendedSaveGameInfoDelegate> OldDelegate;
 		FDelegateHandle Handle;
 	};
+
+	int32 MaxEvents = 32 * 1024;
+	static FAutoConsoleVariableRef CVarMaxEvents(TEXT("SaveGameReplay.MaxEvents"), MaxEvents, TEXT("Maximum number of expected events per replay."));
+
+	int32 MaxEventSize = 128 * 1024;
+	static FAutoConsoleVariableRef CVarMaxEventSize(TEXT("SaveGameReplay.MaxEventSize"), MaxEventSize, TEXT("Maximum expected size of a replay event in bytes.."));
 }
 
 struct FConstConsoleVars
@@ -1598,6 +1604,12 @@ bool FSaveGameNetworkReplayStreamer::ReadMetaDataFromLocalStream(FArchive& Strea
 
 		if (LocalEvent.SizeInBytes > 0)
 		{
+			if (LocalEvent.SizeInBytes > SaveGameReplay::MaxEventSize)
+			{
+				UE_LOG(LogSaveGameReplay, Error, TEXT("ReadMetaDataFromLocalStream: Event SizeInBytes > MaxEventSize (%d / %d)"), LocalEvent.SizeInBytes, SaveGameReplay::MaxEventSize);
+				return false;
+			}
+
 			ReplayEventData.SetNumUninitialized(LocalEvent.SizeInBytes);
 			StreamArchive.Seek(LocalEvent.EventDataOffset);
 			StreamArchive.Serialize(ReplayEventData.GetData(), ReplayEventData.Num());
@@ -1728,6 +1740,13 @@ bool FSaveGameNetworkReplayStreamer::SerializeVersionedMetaData(FArchive& Archiv
 		int32 NumEvents = Events.Num();
 		Archive << NumEvents;
 
+		if (NumEvents > SaveGameReplay::MaxEvents)
+		{
+			UE_LOG(LogSaveGameReplay, Error, TEXT("SerializeVersionedMetaData: NumEvents > MaxEvents (%d / %d)"), NumEvents, SaveGameReplay::MaxEvents);
+			Archive.SetError();
+			return false;
+		}
+
 		Events.SetNum(NumEvents);
 		EventData.SetNum(NumEvents);
 
@@ -1747,6 +1766,13 @@ bool FSaveGameNetworkReplayStreamer::SerializeVersionedMetaData(FArchive& Archiv
 
 			if (EventDataSize > 0)
 			{
+				if (EventDataSize > SaveGameReplay::MaxEventSize)
+				{
+					UE_LOG(LogSaveGameReplay, Error, TEXT("SerializeVersionedMetaData: EventDataSize > MaxEventSize (%d / %d)"), EventDataSize, SaveGameReplay::MaxEventSize);
+					Archive.SetError();
+					return false;
+				}
+
 				ReplayEventData.SetNum(EventDataSize);
 				Archive.Serialize(ReplayEventData.GetData(), ReplayEventData.Num());
 			}
