Clear spectator controller references in StopDemo
Add some additional error logs in local file streamer

diff --git a/Engine/Source/Runtime/Engine/Private/DemoNetDriver.cpp b/Engine/Source/Runtime/Engine/Private/DemoNetDriver.cpp
--- a/Engine/Source/Runtime/Engine/Private/DemoNetDriver.cpp
+++ b/Engine/Source/Runtime/Engine/Private/DemoNetDriver.cpp
@@ -1536,6 +1536,9 @@ void UDemoNetDriver::StopDemo()
 
 	ReplayHelper.StopReplay();
 
+	SpectatorControllers.Empty();
+	SpectatorController = nullptr;
+
 	ClearReplayTasks();
 	ResetDemoState();
 
@@ -3121,8 +3124,8 @@ int32 UDemoNetDriver::CleanUpSplitscreenConnections(bool bDeleteOwner)
 		}
 	}
 
-	FString OwnerDeletionStr(bDeleteOwner ? TEXT("with") : TEXT("without"));
-	UE_LOG(LogDemo, Log, TEXT("Cleaned up %d splitscreen connections %s owner deletion"), NumSplitscreenConnectionsCleaned, *OwnerDeletionStr);
+	UE_LOG(LogDemo, Log, TEXT("Cleaned up %d splitscreen connections, owner deletion: %s"), NumSplitscreenConnectionsCleaned, bDeleteOwner ? TEXT("enabled") : TEXT("disabled"));
+
 	return NumSplitscreenConnectionsCleaned;
 }
 
diff --git a/Engine/Source/Runtime/NetworkReplayStreaming/LocalFileNetworkReplayStreaming/Private/LocalFileNetworkReplayStreaming.cpp b/Engine/Source/Runtime/NetworkReplayStreaming/LocalFileNetworkReplayStreaming/Private/LocalFileNetworkReplayStreaming.cpp
--- a/Engine/Source/Runtime/NetworkReplayStreaming/LocalFileNetworkReplayStreaming/Private/LocalFileNetworkReplayStreaming.cpp
+++ b/Engine/Source/Runtime/NetworkReplayStreaming/LocalFileNetworkReplayStreaming/Private/LocalFileNetworkReplayStreaming.cpp
@@ -1811,6 +1811,7 @@ void FLocalFileNetworkReplayStreamer::FlushStream(const uint32 TimeInMS)
 			}
 			else
 			{
+				UE_LOG(LogLocalFileReplay, Error, TEXT("FLocalFileNetworkReplayStreamer::FlushStream failed with invalid replay info."));
 				SetLastError(ENetworkReplayError::ServiceUnavailable);
 			}
 		});
@@ -1963,6 +1964,7 @@ void FLocalFileNetworkReplayStreamer::FlushCheckpointInternal(const uint32 TimeI
 			}
 			else
 			{
+				UE_LOG(LogLocalFileReplay, Error, TEXT("FLocalFileNetworkReplayStreamer::FlushCheckpointInternal failed with invalid replay info."));
 				SetLastError(ENetworkReplayError::ServiceUnavailable);
 			}
 		});
@@ -3002,6 +3004,7 @@ void FLocalFileNetworkReplayStreamer::WriteHeader()
 			}
 			else
 			{
+				UE_LOG(LogLocalFileReplay, Error, TEXT("FLocalFileNetworkReplayStreamer::WriteHeader failed with invalid replay info."));
 				SetLastError(ENetworkReplayError::ServiceUnavailable);
 			}
 		});
@@ -3233,6 +3236,8 @@ void FLocalFileNetworkReplayStreamer::ConditionallyLoadNextChunk()
 			// Make sure our stream chunk index didn't change under our feet
 			if (RequestedStreamChunkIndex != StreamChunkIndex)
 			{
+				UE_LOG(LogLocalFileReplay, Warning, TEXT("FLocalFileNetworkReplayStreamer::ConditionallyLoadNextChunk failed with requested chunk index mismatch."));
+
 				StreamAr.Reset();
 				SetLastError(ENetworkReplayError::ServiceUnavailable);
 				return;
