Handle potential crash on invalid bad checkpoint metadata
Handle potential crash on invalid read lengths from local stream archive

diff --git a/Engine/Source/Runtime/NetworkReplayStreaming/LocalFileNetworkReplayStreaming/Private/LocalFileNetworkReplayStreaming.cpp b/Engine/Source/Runtime/NetworkReplayStreaming/LocalFileNetworkReplayStreaming/Private/LocalFileNetworkReplayStreaming.cpp
--- a/Engine/Source/Runtime/NetworkReplayStreaming/LocalFileNetworkReplayStreaming/Private/LocalFileNetworkReplayStreaming.cpp
+++ b/Engine/Source/Runtime/NetworkReplayStreaming/LocalFileNetworkReplayStreaming/Private/LocalFileNetworkReplayStreaming.cpp
@@ -88,7 +88,7 @@ void FLocalFileStreamFArchive::Serialize(void* V, int64 Length)
 {
 	if (IsLoading())
 	{
-		if ((Pos + Length) > Buffer.Num())
+		if ((Length < 0) || (Pos + Length) > Buffer.Num())
 		{
 			UE_LOG(LogLocalFileReplay, Error, TEXT("FLocalFileStreamFArchive::Serialize: Attempted to serialize past end of archive: Position = %i, Size=%i, Requested = %lli"), Pos, Buffer.Num(), Length);
 			SetError();
@@ -2170,35 +2170,44 @@ void FLocalFileNetworkReplayStreamer::GotoCheckpointIndex(const int32 Checkpoint
 			CheckpointAr.Buffer = MoveTemp(RequestData.DataBuffer);
 			CheckpointAr.Pos = 0;
 
-			int32 DataChunkIndex = FCString::Atoi(*CurrentReplayInfo.Checkpoints[CheckpointIndex].Metadata);
-			check(CurrentReplayInfo.DataChunks.IsValidIndex(DataChunkIndex));
+			const bool bIsDataAvailableForTimeRange = IsDataAvailableForTimeRange(CurrentReplayInfo.Checkpoints[CheckpointIndex].Time1, LastGotoTimeInMS);
 
-			bool bIsDataAvailableForTimeRange = IsDataAvailableForTimeRange(CurrentReplayInfo.Checkpoints[CheckpointIndex].Time1, LastGotoTimeInMS);
-
-			if (!bIsDataAvailableForTimeRange)
+			const int32 DataChunkIndex = FCString::Atoi(*CurrentReplayInfo.Checkpoints[CheckpointIndex].Metadata);
+			
+			if (CurrentReplayInfo.DataChunks.IsValidIndex(DataChunkIndex))
 			{
-				// Completely reset our stream (we're going to start loading from the start of the checkpoint)
-				StreamAr.Reset();
+				if (!bIsDataAvailableForTimeRange)
+				{
+					// Completely reset our stream (we're going to start loading from the start of the checkpoint)
+					StreamAr.Reset();
 
-				// Reset any time we were waiting on in the past
-				HighPriorityEndTime = 0;
+					// Reset any time we were waiting on in the past
+					HighPriorityEndTime = 0;
 
-				StreamDataOffset = CurrentReplayInfo.DataChunks[DataChunkIndex].StreamOffset;
+					StreamDataOffset = CurrentReplayInfo.DataChunks[DataChunkIndex].StreamOffset;
 
-				// Reset our stream range
-				StreamTimeRange = TInterval<uint32>(0, 0);
+					// Reset our stream range
+					StreamTimeRange = TInterval<uint32>(0, 0);
 
-				// Set the next chunk to be right after this checkpoint (which was stored in the metadata)
-				StreamChunkIndex = DataChunkIndex;
+					// Set the next chunk to be right after this checkpoint (which was stored in the metadata)
+					StreamChunkIndex = DataChunkIndex;
 
-				LastChunkTime = 0;		// Force the next chunk to start loading immediately in case LastGotoTimeInMS is 0 (which would effectively disable high priority mode immediately)
+					LastChunkTime = 0;		// Force the next chunk to start loading immediately in case LastGotoTimeInMS is 0 (which would effectively disable high priority mode immediately)
+				}
+				else
+				{
+					// set stream position back to the correct location
+					StreamAr.Pos = CurrentReplayInfo.DataChunks[DataChunkIndex].StreamOffset - StreamDataOffset;
+					check(StreamAr.Pos >= 0 && StreamAr.Pos <= StreamAr.Buffer.Num());
+					StreamAr.bAtEndOfReplay = false;
+				}
 			}
 			else
 			{
-				// set stream position back to the correct location
-				StreamAr.Pos = CurrentReplayInfo.DataChunks[DataChunkIndex].StreamOffset - StreamDataOffset;
-				check(StreamAr.Pos >= 0 && StreamAr.Pos <= StreamAr.Buffer.Num());
-				StreamAr.bAtEndOfReplay = false;
+				UE_LOG(LogLocalFileReplay, Error, TEXT("FLocalFileNetworkReplayStreamer::GotoCheckpointIndexDelta. Checkpoint data chunk index invalid: %d"), DataChunkIndex);
+				Delegate.ExecuteIfBound(RequestData.DelegateResult);
+				LastGotoTimeInMS = -1;
+				return;
 			}
 
 			// If we want to fast forward past the end of a stream (and we set a new chunk to stream), clamp to the checkpoint
@@ -2340,7 +2349,7 @@ void FLocalFileNetworkReplayStreamer::GotoCheckpointIndex(const int32 Checkpoint
 			CheckpointAr.Pos = 0;
 
 			const FLocalFileEventInfo& Checkpoint = CurrentReplayInfo.Checkpoints[CheckpointIndex];
-			int32 DataChunkIndex = FCString::Atoi(*Checkpoint.Metadata);
+			const int32 DataChunkIndex = FCString::Atoi(*Checkpoint.Metadata);
 
 			if (CurrentReplayInfo.DataChunks.IsValidIndex(DataChunkIndex))
 			{
