Moving save game replay file helper commands to the task graph

diff --git a/Engine/Source/Runtime/NetworkReplayStreaming/SaveGameNetworkReplayStreaming/Private/SaveGameNetworkReplayStreaming.cpp b/Engine/Source/Runtime/NetworkReplayStreaming/SaveGameNetworkReplayStreaming/Private/SaveGameNetworkReplayStreaming.cpp
--- a/Engine/Source/Runtime/NetworkReplayStreaming/SaveGameNetworkReplayStreaming/Private/SaveGameNetworkReplayStreaming.cpp
+++ b/Engine/Source/Runtime/NetworkReplayStreaming/SaveGameNetworkReplayStreaming/Private/SaveGameNetworkReplayStreaming.cpp
@@ -248,7 +248,7 @@ namespace SaveGameReplay
 		{
 			if (IsRunningDedicatedServer())
 			{
-				UE_LOG(LogSaveGameReplay, Warning, TEXT("FSaveGameMoveFileHelper commands are client only."));
+				UE_LOG(LogSaveGameReplay, Warning, TEXT("FSaveGameReplayMoveFileHelper commands are client only."));
 				return;
 			}
 
@@ -259,7 +259,7 @@ namespace SaveGameReplay
 			}
 			else if (Params.Num() != 0)
 			{
-				UE_LOG(LogSaveGameReplay, Warning, TEXT("FSaveGameMoveFileHelper commands take either a Streamer Override or no arguments."));
+				UE_LOG(LogSaveGameReplay, Warning, TEXT("FSaveGameReplayMoveFileHelper commands take either a Streamer Override or no arguments."));
 				return;
 			}
 
@@ -277,7 +277,15 @@ namespace SaveGameReplay
 					const int32 UserIndex = GetFirstPlayerIndex();
 					if (ensureMsgf(INDEX_NONE != UserIndex, TEXT("FSaveGameReplayMoveFileHelper Unable to get UserIndex")))
 					{
-						CommandToRun(Streamer, DemoPath, UserIndex);
+						Async(EAsyncExecution::TaskGraph, [CommandToRun, Streamer, DemoPath, UserIndex]()
+							{
+								CommandToRun(Streamer, DemoPath, UserIndex);
+							},
+							[]()
+							{
+								UE_LOG(LogSaveGameReplay, Log, TEXT("FSaveGameReplayMoveFileHelper: Async RunCommand complete."));
+							}
+						);
 					}
 				}
 			}
@@ -338,7 +346,7 @@ namespace SaveGameReplay
 
 			for (const FString& CurrentName : FoundFiles)
 			{
-				UE_LOG(LogSaveGameReplay, Log, TEXT("FSaveGameMoveFileHelper::SanitizeNames - Handling %s"), *CurrentName);
+				UE_LOG(LogSaveGameReplay, Log, TEXT("FSaveGameReplayMoveFileHelper::SanitizeNames - Handling %s"), *CurrentName);
 
 				FString NewName(CurrentName);
 				NewName.RemoveFromEnd(LocalFileExt);
@@ -346,7 +354,7 @@ namespace SaveGameReplay
 
 				if (!FileManager.Move(*FPaths::Combine(DemoPath, NewName), *FPaths::Combine(DemoPath, CurrentName)))
 				{
-					UE_LOG(LogSaveGameReplay, Warning, TEXT("FSaveGameMoveFileHelper::SanitizeNames - Failed to sanitize %s"), *CurrentName);
+					UE_LOG(LogSaveGameReplay, Warning, TEXT("FSaveGameReplayMoveFileHelper::SanitizeNames - Failed to sanitize %s"), *CurrentName);
 				}
 			}
 		}
@@ -392,10 +400,10 @@ namespace SaveGameReplay
 			const uint32 Result = IFileManager::Get().Copy(*DestinationFileName, *SourceFileName);
 			if (0 != Result)
 			{
-				UE_LOG(LogSaveGameReplay, Warning, TEXT("FSaveGameMoveFileHelper::CopyFile: Failed - from '%s' to '%s' error = %lu"), *SourceFileName, *DestinationFileName, Result);
+				UE_LOG(LogSaveGameReplay, Warning, TEXT("FSaveGameReplayMoveFileHelper::CopyFile: Failed - from '%s' to '%s' error = %lu"), *SourceFileName, *DestinationFileName, Result);
 			}
 
-			UE_LOG(LogSaveGameReplay, Log, TEXT("FSaveGameMoveFileHelper::CopyFile: Result = %d"), Result);
+			UE_LOG(LogSaveGameReplay, Log, TEXT("FSaveGameReplayMoveFileHelper::CopyFile: Result = %d"), Result);
 		}
 
 		static void SaveFile(const TSharedPtr<FMoveContext>& Context, const FString& SaveGameName, ISaveGameSystem* SaveGameSystem)
@@ -403,7 +411,7 @@ namespace SaveGameReplay
 			TArray<uint8> SaveData;
 			if (!SaveGameSystem->LoadGame(false, *SaveGameName, Context->UserIndex, SaveData))
 			{
-				UE_LOG(LogSaveGameReplay, Warning, TEXT("FSaveGameMoveFileHelper::SaveFile: Failed to load save game %s"), *SaveGameName)
+				UE_LOG(LogSaveGameReplay, Warning, TEXT("FSaveGameReplayMoveFileHelper::SaveFile: Failed to load save game %s"), *SaveGameName)
 					return;
 			}
 
@@ -415,7 +423,7 @@ namespace SaveGameReplay
 
 			if (FileAR->IsError())
 			{
-				UE_LOG(LogSaveGameReplay, Warning, TEXT("FSaveGameMoveFileHelper::SaveFile: Failed to save game %s to %s"), *SaveGameName, *DestinationFileName);
+				UE_LOG(LogSaveGameReplay, Warning, TEXT("FSaveGameReplayMoveFileHelper::SaveFile: Failed to save game %s to %s"), *SaveGameName, *DestinationFileName);
 			}
 		}
 
@@ -440,7 +448,7 @@ namespace SaveGameReplay
 						}
 						else
 						{
-							UE_LOG(LogSaveGameReplay, Warning, TEXT("FSaveGameMoveFileHelper::OnEnumerateRecentStreamsComplete: Unable to get SaveGameSystem"));
+							UE_LOG(LogSaveGameReplay, Warning, TEXT("FSaveGameReplayMoveFileHelper::OnEnumerateRecentStreamsComplete: Unable to get SaveGameSystem"));
 						}
 					}
 					else
@@ -454,7 +462,7 @@ namespace SaveGameReplay
 			}
 			else
 			{
-				UE_LOG(LogSaveGameReplay, Warning, TEXT("FSaveGameMoveFileHelper::OnEnumerateRecentStreamsComplete: Enumerate failed"));
+				UE_LOG(LogSaveGameReplay, Warning, TEXT("FSaveGameReplayMoveFileHelper::OnEnumerateRecentStreamsComplete: Enumerate failed"));
 			}
 
 			MoveFilesInternal_PostEnumerate(Context);
@@ -473,7 +481,7 @@ namespace SaveGameReplay
 			}
 			else
 			{
-				UE_LOG(LogSaveGameReplay, Warning, TEXT("FSaveGameMoveFileHelper::OnEnumerateRecentStreamsComplete: Enumerate failed"))
+				UE_LOG(LogSaveGameReplay, Warning, TEXT("FSaveGameReplayMoveFileHelper::OnEnumerateRecentStreamsComplete: Enumerate failed"))
 			}
 		}
 
@@ -497,7 +505,7 @@ namespace SaveGameReplay
 #if PLATFORM_MOVE_REQUIRES_LOWERCASE
 				if (ReplayFileName.Compare(ReplayFileName.ToLower(), ESearchCase::CaseSensitive) != 0)
 				{
-					UE_LOG(LogSaveGameReplay, Warning, TEXT("FSaveGameMoveFileHelper::MoveFilesLocalInternal - Replay file %s is not lowercase, import will fail."), *ReplayFileName);
+					UE_LOG(LogSaveGameReplay, Warning, TEXT("FSaveGameReplayMoveFileHelper::MoveFilesLocalInternal - Replay file %s is not lowercase, import will fail."), *ReplayFileName);
 				}
 #endif
 
