Backing out async change to debug commands

diff --git a/Engine/Source/Runtime/NetworkReplayStreaming/SaveGameNetworkReplayStreaming/Private/SaveGameNetworkReplayStreaming.cpp b/Engine/Source/Runtime/NetworkReplayStreaming/SaveGameNetworkReplayStreaming/Private/SaveGameNetworkReplayStreaming.cpp
--- a/Engine/Source/Runtime/NetworkReplayStreaming/SaveGameNetworkReplayStreaming/Private/SaveGameNetworkReplayStreaming.cpp
+++ b/Engine/Source/Runtime/NetworkReplayStreaming/SaveGameNetworkReplayStreaming/Private/SaveGameNetworkReplayStreaming.cpp
@@ -248,7 +248,7 @@ namespace SaveGameReplay
 		{
 			if (IsRunningDedicatedServer())
 			{
-				UE_LOG(LogSaveGameReplay, Warning, TEXT("FSaveGameReplayMoveFileHelper commands are client only."));
+				UE_LOG(LogSaveGameReplay, Warning, TEXT("FSaveGameMoveFileHelper commands are client only."));
 				return;
 			}
 
@@ -259,7 +259,7 @@ namespace SaveGameReplay
 			}
 			else if (Params.Num() != 0)
 			{
-				UE_LOG(LogSaveGameReplay, Warning, TEXT("FSaveGameReplayMoveFileHelper commands take either a Streamer Override or no arguments."));
+				UE_LOG(LogSaveGameReplay, Warning, TEXT("FSaveGameMoveFileHelper commands take either a Streamer Override or no arguments."));
 				return;
 			}
 
@@ -277,15 +277,7 @@ namespace SaveGameReplay
 					const int32 UserIndex = GetFirstPlayerIndex();
 					if (ensureMsgf(INDEX_NONE != UserIndex, TEXT("FSaveGameReplayMoveFileHelper Unable to get UserIndex")))
 					{
-						Async(EAsyncExecution::TaskGraph, [CommandToRun, Streamer, DemoPath, UserIndex]()
-							{
-								CommandToRun(Streamer, DemoPath, UserIndex);
-							},
-							[]()
-							{
-								UE_LOG(LogSaveGameReplay, Log, TEXT("FSaveGameReplayMoveFileHelper: Async RunCommand complete."));
-							}
-						);
+						CommandToRun(Streamer, DemoPath, UserIndex);
 					}
 				}
 			}
@@ -346,7 +338,7 @@ namespace SaveGameReplay
 
 			for (const FString& CurrentName : FoundFiles)
 			{
-				UE_LOG(LogSaveGameReplay, Log, TEXT("FSaveGameReplayMoveFileHelper::SanitizeNames - Handling %s"), *CurrentName);
+				UE_LOG(LogSaveGameReplay, Log, TEXT("FSaveGameMoveFileHelper::SanitizeNames - Handling %s"), *CurrentName);
 
 				FString NewName(CurrentName);
 				NewName.RemoveFromEnd(LocalFileExt);
@@ -354,7 +346,7 @@ namespace SaveGameReplay
 
 				if (!FileManager.Move(*FPaths::Combine(DemoPath, NewName), *FPaths::Combine(DemoPath, CurrentName)))
 				{
-					UE_LOG(LogSaveGameReplay, Warning, TEXT("FSaveGameReplayMoveFileHelper::SanitizeNames - Failed to sanitize %s"), *CurrentName);
+					UE_LOG(LogSaveGameReplay, Warning, TEXT("FSaveGameMoveFileHelper::SanitizeNames - Failed to sanitize %s"), *CurrentName);
 				}
 			}
 		}
@@ -400,10 +392,10 @@ namespace SaveGameReplay
 			const uint32 Result = IFileManager::Get().Copy(*DestinationFileName, *SourceFileName);
 			if (0 != Result)
 			{
-				UE_LOG(LogSaveGameReplay, Warning, TEXT("FSaveGameReplayMoveFileHelper::CopyFile: Failed - from '%s' to '%s' error = %lu"), *SourceFileName, *DestinationFileName, Result);
+				UE_LOG(LogSaveGameReplay, Warning, TEXT("FSaveGameMoveFileHelper::CopyFile: Failed - from '%s' to '%s' error = %lu"), *SourceFileName, *DestinationFileName, Result);
 			}
 
-			UE_LOG(LogSaveGameReplay, Log, TEXT("FSaveGameReplayMoveFileHelper::CopyFile: Result = %d"), Result);
+			UE_LOG(LogSaveGameReplay, Log, TEXT("FSaveGameMoveFileHelper::CopyFile: Result = %d"), Result);
 		}
 
 		static void SaveFile(const TSharedPtr<FMoveContext>& Context, const FString& SaveGameName, ISaveGameSystem* SaveGameSystem)
@@ -411,7 +403,7 @@ namespace SaveGameReplay
 			TArray<uint8> SaveData;
 			if (!SaveGameSystem->LoadGame(false, *SaveGameName, Context->UserIndex, SaveData))
 			{
-				UE_LOG(LogSaveGameReplay, Warning, TEXT("FSaveGameReplayMoveFileHelper::SaveFile: Failed to load save game %s"), *SaveGameName)
+				UE_LOG(LogSaveGameReplay, Warning, TEXT("FSaveGameMoveFileHelper::SaveFile: Failed to load save game %s"), *SaveGameName)
 					return;
 			}
 
@@ -423,7 +415,7 @@ namespace SaveGameReplay
 
 			if (FileAR->IsError())
 			{
-				UE_LOG(LogSaveGameReplay, Warning, TEXT("FSaveGameReplayMoveFileHelper::SaveFile: Failed to save game %s to %s"), *SaveGameName, *DestinationFileName);
+				UE_LOG(LogSaveGameReplay, Warning, TEXT("FSaveGameMoveFileHelper::SaveFile: Failed to save game %s to %s"), *SaveGameName, *DestinationFileName);
 			}
 		}
 
@@ -448,11 +440,12 @@ namespace SaveGameReplay
 						}
 						else
 						{
-							UE_LOG(LogSaveGameReplay, Warning, TEXT("FSaveGameReplayMoveFileHelper::OnEnumerateRecentStreamsComplete: Unable to get SaveGameSystem"));
+							UE_LOG(LogSaveGameReplay, Warning, TEXT("FSaveGameMoveFileHelper::OnEnumerateRecentStreamsComplete: Unable to get SaveGameSystem"));
 						}
 					}
 					else
 					{
+						//@todo: make async, but don't forget about post enumerate
 						for (const FNetworkReplayStreamInfo& StreamInfo : Result.FoundStreams)
 						{
 							CopyFile(Context, StreamInfo.Name);
@@ -462,7 +455,7 @@ namespace SaveGameReplay
 			}
 			else
 			{
-				UE_LOG(LogSaveGameReplay, Warning, TEXT("FSaveGameReplayMoveFileHelper::OnEnumerateRecentStreamsComplete: Enumerate failed"));
+				UE_LOG(LogSaveGameReplay, Warning, TEXT("FSaveGameMoveFileHelper::OnEnumerateRecentStreamsComplete: Enumerate failed"));
 			}
 
 			MoveFilesInternal_PostEnumerate(Context);
@@ -472,6 +465,8 @@ namespace SaveGameReplay
 		{
 			if (Result.WasSuccessful())
 			{
+				//@todo: make async
+
 				// Currently, the LocalFileStreamer doesn't support EnumerateRecentStreams
 				// and the SaveGameStreamer will just return non-Saved replays.
 				for (const FNetworkReplayStreamInfo& StreamInfo : Result.FoundStreams)
@@ -481,7 +476,7 @@ namespace SaveGameReplay
 			}
 			else
 			{
-				UE_LOG(LogSaveGameReplay, Warning, TEXT("FSaveGameReplayMoveFileHelper::OnEnumerateRecentStreamsComplete: Enumerate failed"))
+				UE_LOG(LogSaveGameReplay, Warning, TEXT("FSaveGameMoveFileHelper::OnEnumerateRecentStreamsComplete: Enumerate failed"))
 			}
 		}
 
@@ -505,7 +500,7 @@ namespace SaveGameReplay
 #if PLATFORM_MOVE_REQUIRES_LOWERCASE
 				if (ReplayFileName.Compare(ReplayFileName.ToLower(), ESearchCase::CaseSensitive) != 0)
 				{
-					UE_LOG(LogSaveGameReplay, Warning, TEXT("FSaveGameReplayMoveFileHelper::MoveFilesLocalInternal - Replay file %s is not lowercase, import will fail."), *ReplayFileName);
+					UE_LOG(LogSaveGameReplay, Warning, TEXT("FSaveGameMoveFileHelper::MoveFilesLocalInternal - Replay file %s is not lowercase, import will fail."), *ReplayFileName);
 				}
 #endif
 
@@ -527,6 +522,7 @@ namespace SaveGameReplay
 			MoveFiles(Streamer, GetTempDemoDirectory(), SourceDirectory, UserIndex);
 		}
 
+		//@todo: make async
 		static void MoveFilesFromTemp(const FString& DestinationDirectory, const int32 UserIndex)
 		{
 			TSharedPtr<FMoveContext> MoveContext(MakeShareable(new FMoveContext(nullptr, GetTempDemoDirectory(), DestinationDirectory, UserIndex)));
