Clang static analysis fix

diff --git a/Engine/Source/Runtime/NetworkReplayStreaming/SaveGameNetworkReplayStreaming/Private/SaveGameNetworkReplayStreaming.cpp b/Engine/Source/Runtime/NetworkReplayStreaming/SaveGameNetworkReplayStreaming/Private/SaveGameNetworkReplayStreaming.cpp
--- a/Engine/Source/Runtime/NetworkReplayStreaming/SaveGameNetworkReplayStreaming/Private/SaveGameNetworkReplayStreaming.cpp
+++ b/Engine/Source/Runtime/NetworkReplayStreaming/SaveGameNetworkReplayStreaming/Private/SaveGameNetworkReplayStreaming.cpp
@@ -177,7 +177,18 @@ namespace SaveGameReplay
 				}
 				else
 				{
-					OutstandingTask = MakeUnique<TAsyncTask<TResult>>(OwningStreamer, Description, InAsyncWork, InPostAsyncWork);
+					if (InAsyncWork)
+					{
+						OutstandingTask = MakeUnique<TAsyncTask<TResult>>(OwningStreamer, Description, InAsyncWork, InPostAsyncWork);
+					}
+					else
+					{
+						UE_LOG(LogSaveGameReplay, Warning, TEXT("SaveGameReplay::FAsyncTaskManager::StartTask - Async work function was invalid (NewTask = %s)"), *Description);
+
+						TResult Result;
+						Result.Result = EStreamingOperationResult::UnfinishedTask;
+						InPostAsyncWork(Result);
+					}
 				}
 			}
 		}
