Changed engine/game network version to be custom versions
Changed replay header version to be a custom version
Changed local file replay streamer version to be a custom version

diff --git a/Engine/Plugins/Runtime/GameplayAbilities/Source/GameplayAbilities/Private/GameplayAbilityRepAnimMontage.cpp b/Engine/Plugins/Runtime/GameplayAbilities/Source/GameplayAbilities/Private/GameplayAbilityRepAnimMontage.cpp
--- a/Engine/Plugins/Runtime/GameplayAbilities/Source/GameplayAbilities/Private/GameplayAbilityRepAnimMontage.cpp
+++ b/Engine/Plugins/Runtime/GameplayAbilities/Source/GameplayAbilities/Private/GameplayAbilityRepAnimMontage.cpp
@@ -6,6 +6,8 @@
 
 bool FGameplayAbilityRepAnimMontage::NetSerialize(FArchive& Ar, class UPackageMap* Map, bool& bOutSuccess)
 {
+	Ar.UsingCustomVersion(FEngineNetworkCustomVersion::Guid);
+
 	uint8 RepPosition = bRepPosition;
 	Ar.SerializeBits(&RepPosition, 1);
 	if (RepPosition)
@@ -38,7 +40,7 @@ bool FGameplayAbilityRepAnimMontage::NetSerialize(FArchive& Ar, class UPackageMa
 	Ar.SerializeBits(&bIsStopped, 1);
 	IsStopped = bIsStopped & 1;
 
-	if (Ar.EngineNetVer() < HISTORY_MONTAGE_PLAY_INST_ID_SERIALIZATION)
+	if (Ar.EngineNetVer() < FEngineNetworkCustomVersion::MontagePlayInstIdSerialization)
 	{
 		uint8 bForcePlayBit = 0;
 		Ar.SerializeBits(&bForcePlayBit, 1);
@@ -62,7 +64,7 @@ bool FGameplayAbilityRepAnimMontage::NetSerialize(FArchive& Ar, class UPackageMa
 	Ar << PlayRate;
 	Ar << BlendTime;
 	Ar << NextSectionID;
-	if (Ar.EngineNetVer() >= HISTORY_MONTAGE_PLAY_INST_ID_SERIALIZATION)
+	if (Ar.EngineNetVer() >= FEngineNetworkCustomVersion::MontagePlayInstIdSerialization)
 	{
 		Ar << PlayInstanceId;
 	}
diff --git a/Engine/Source/Runtime/Core/Private/Math/UnrealMath.cpp b/Engine/Source/Runtime/Core/Private/Math/UnrealMath.cpp
--- a/Engine/Source/Runtime/Core/Private/Math/UnrealMath.cpp
+++ b/Engine/Source/Runtime/Core/Private/Math/UnrealMath.cpp
@@ -1130,6 +1130,8 @@ namespace Math
 	template<typename T>
 	bool TQuat<T>::NetSerialize(FArchive& Ar, class UPackageMap*, bool& bOutSuccess)
 	{
+		Ar.UsingCustomVersion(FEngineNetworkCustomVersion::Guid);
+
 		TQuat<T> Q;
 
 		if (Ar.IsSaving())
@@ -1159,7 +1161,7 @@ namespace Math
 			}
 		}
 
-		if (Ar.EngineNetVer() >= HISTORY_SERIALIZE_DOUBLE_VECTORS_AS_DOUBLES && Ar.EngineNetVer() != HISTORY_21_AND_VIEWPITCH_ONLY_DO_NOT_USE)
+		if (Ar.EngineNetVer() >= FEngineNetworkCustomVersion::SerializeDoubleVectorsAsDoubles && Ar.EngineNetVer() != FEngineNetworkCustomVersion::Ver21AndViewPitchOnly_DONOTUSE)
 		{
 			Ar << Q.X << Q.Y << Q.Z;
 		}
diff --git a/Engine/Source/Runtime/Core/Private/Misc/NetworkVersion.cpp b/Engine/Source/Runtime/Core/Private/Misc/NetworkVersion.cpp
--- a/Engine/Source/Runtime/Core/Private/Misc/NetworkVersion.cpp
+++ b/Engine/Source/Runtime/Core/Private/Misc/NetworkVersion.cpp
@@ -11,8 +11,10 @@
 #include "Misc/CommandLine.h"
 #include "Misc/Crc.h"
 #include "Misc/EngineVersion.h"
+#include "Misc/Guid.h"
 #include "Misc/NetworkGuid.h" // IWYU pragma: keep
 #include "Misc/Parse.h"
+#include "Serialization/CustomVersion.h"
 #include "Trace/Detail/Channel.h"
 #include "Misc/StringBuilder.h"
 
@@ -22,6 +24,12 @@ FNetworkVersion::FGetLocalNetworkVersionOverride FNetworkVersion::GetLocalNetwor
 FNetworkVersion::FIsNetworkCompatibleOverride FNetworkVersion::IsNetworkCompatibleOverride;
 FNetworkVersion::FGetReplayCompatibleChangeListOverride FNetworkVersion::GetReplayCompatibleChangeListOverride;
 
+const FGuid FEngineNetworkCustomVersion::Guid = FGuid(0x62915CA3, 0x1C8E4BF7, 0xA30E12C7, 0xC8219DF7);
+FCustomVersionRegistration GRegisterEngineNetworkCustomVersion(FEngineNetworkCustomVersion::Guid, FEngineNetworkCustomVersion::LatestVersion, TEXT("EngineNetworkVersion"));
+
+const FGuid FGameNetworkCustomVersion::Guid = FGuid(0xCC400D24, 0xE0E94E7B, 0x9BF9A283, 0xDCC0C027);
+FCustomVersionRegistration GRegisterGameNetworkCustomVersion(FGameNetworkCustomVersion::Guid, FGameNetworkCustomVersion::LatestVersion, TEXT("GameNetworkVersion"));
+
 FString& FNetworkVersion::GetProjectVersion_Internal()
 {
 	static FString ProjectVersion = TEXT("1.0.0");
@@ -31,12 +39,20 @@ FString& FNetworkVersion::GetProjectVersion_Internal()
 bool FNetworkVersion::bHasCachedNetworkChecksum			= false;
 uint32 FNetworkVersion::CachedNetworkChecksum			= 0;
 
-uint32 FNetworkVersion::EngineNetworkProtocolVersion	= HISTORY_ENGINENETVERSION_LATEST;
+bool FNetworkVersion::bHasCachedReplayChecksum = false;
+uint32 FNetworkVersion::CachedReplayChecksum = 0;
+
+uint32 FNetworkVersion::EngineNetworkProtocolVersion	= FEngineNetworkCustomVersion::LatestVersion;
 uint32 FNetworkVersion::GameNetworkProtocolVersion		= 0;
 
-uint32 FNetworkVersion::EngineCompatibleNetworkProtocolVersion		= HISTORY_REPLAY_BACKWARDS_COMPAT;
+uint32 FNetworkVersion::EngineCompatibleNetworkProtocolVersion		= FEngineNetworkCustomVersion::ReplayBackwardsCompat;
 uint32 FNetworkVersion::GameCompatibleNetworkProtocolVersion		= 0;
 
+FCustomVersionContainer FNetworkVersion::NetworkCustomVersions;
+FCustomVersionContainer FNetworkVersion::CompatibleNetworkCustomVersions;
+
+bool FNetworkVersion::bInitializedCustomVersions = false;
+
 void FNetworkVersion::SetProjectVersion(const TCHAR* InVersion)
 {
 	if (ensureMsgf(InVersion != nullptr && FCString::Strlen(InVersion), TEXT("ProjectVersion used for network version must be a valid string!")))
@@ -52,18 +68,30 @@ void FNetworkVersion::SetProjectVersion(const TCHAR* InVersion)
 
 void FNetworkVersion::SetGameNetworkProtocolVersion(const uint32 InGameNetworkProtocolVersion)
 {
+	InitCustomVersions();
+
+	PRAGMA_DISABLE_DEPRECATION_WARNINGS
 	GameNetworkProtocolVersion = InGameNetworkProtocolVersion;
+	PRAGMA_ENABLE_DEPRECATION_WARNINGS
+	NetworkCustomVersions.SetVersion(FGameNetworkCustomVersion::Guid, InGameNetworkProtocolVersion, TEXT("GameNetworkVersion"));
 	bHasCachedNetworkChecksum = false;
 
-	UE_LOG(LogNetVersion, Log, TEXT("Set GameNetworkProtocolVersion to %ud. Version Checksum will be recalculated on next use."), GameNetworkProtocolVersion);
+	UE_LOG(LogNetVersion, Log, TEXT("Set GameNetworkProtocolVersion to %ud. Version Checksum will be recalculated on next use."), InGameNetworkProtocolVersion);
 }
 
 void FNetworkVersion::SetGameCompatibleNetworkProtocolVersion(const uint32 InGameCompatibleNetworkProtocolVersion)
 {
+	InitCustomVersions();
+
+	PRAGMA_DISABLE_DEPRECATION_WARNINGS
 	GameCompatibleNetworkProtocolVersion = InGameCompatibleNetworkProtocolVersion;
+	PRAGMA_ENABLE_DEPRECATION_WARNINGS
+	CompatibleNetworkCustomVersions.SetVersion(FGameNetworkCustomVersion::Guid, InGameCompatibleNetworkProtocolVersion, TEXT("GameNetworkVersion"));
+	
 	bHasCachedNetworkChecksum = false;
+	bHasCachedReplayChecksum = false;
 
-	UE_LOG(LogNetVersion, Log, TEXT("Set GameCompatibleNetworkProtocolVersion to %ud. Version Checksum will be recalculated on next use."), GameCompatibleNetworkProtocolVersion);
+	UE_LOG(LogNetVersion, Log, TEXT("Set GameCompatibleNetworkProtocolVersion to %ud. Version Checksum will be recalculated on next use."), InGameCompatibleNetworkProtocolVersion);
 }
 
 uint32 FNetworkVersion::GetNetworkCompatibleChangelist()
@@ -106,24 +134,90 @@ uint32 FNetworkVersion::GetReplayCompatibleChangelist()
 	return FEngineVersion::CompatibleWith().GetChangelist();
 }
 
+uint32 FNetworkVersion::GetNetworkProtocolVersion(const FGuid& VersionGuid)
+{
+	InitCustomVersions();
+
+	const FCustomVersion* CustomVersion = NetworkCustomVersions.GetVersion(VersionGuid);
+	return CustomVersion ? CustomVersion->Version : 0;
+}
+
+uint32 FNetworkVersion::GetCompatibleNetworkProtocolVersion(const FGuid& VersionGuid)
+{
+	InitCustomVersions();
+
+	const FCustomVersion* CustomVersion = CompatibleNetworkCustomVersions.GetVersion(VersionGuid);
+	return CustomVersion ? CustomVersion->Version : 0;
+}
+
 uint32 FNetworkVersion::GetEngineNetworkProtocolVersion()
 {
-	return EngineNetworkProtocolVersion;
+	return GetNetworkProtocolVersion(FEngineNetworkCustomVersion::Guid);
 }
 
 uint32 FNetworkVersion::GetEngineCompatibleNetworkProtocolVersion()
 {
-	return EngineCompatibleNetworkProtocolVersion;
+	return GetCompatibleNetworkProtocolVersion(FEngineNetworkCustomVersion::Guid);
 }
 
 uint32 FNetworkVersion::GetGameNetworkProtocolVersion()
 {
-	return GameNetworkProtocolVersion;
+	return GetNetworkProtocolVersion(FGameNetworkCustomVersion::Guid);
 }
 
 uint32 FNetworkVersion::GetGameCompatibleNetworkProtocolVersion()
 {
-	return GameCompatibleNetworkProtocolVersion;
+	return GetCompatibleNetworkProtocolVersion(FGameNetworkCustomVersion::Guid);
+}
+
+const FCustomVersionContainer& FNetworkVersion::GetNetworkCustomVersions()
+{
+	InitCustomVersions();
+
+	return NetworkCustomVersions;
+}
+
+void FNetworkVersion::InitCustomVersions()
+{
+	if (!bInitializedCustomVersions)
+	{
+		FCustomVersion RegisteredEngineVersion = FCurrentCustomVersions::Get(FEngineNetworkCustomVersion::Guid).GetValue();
+		FCustomVersion RegisteredGameVersion = FCurrentCustomVersions::Get(FGameNetworkCustomVersion::Guid).GetValue();
+
+		NetworkCustomVersions.SetVersion(FEngineNetworkCustomVersion::Guid, RegisteredEngineVersion.Version, RegisteredEngineVersion.GetFriendlyName());
+		NetworkCustomVersions.SetVersion(FGameNetworkCustomVersion::Guid, RegisteredGameVersion.Version, RegisteredGameVersion.GetFriendlyName());
+
+		CompatibleNetworkCustomVersions.SetVersion(FEngineNetworkCustomVersion::Guid, FEngineNetworkCustomVersion::ReplayBackwardsCompat, RegisteredEngineVersion.GetFriendlyName());
+		CompatibleNetworkCustomVersions.SetVersion(FGameNetworkCustomVersion::Guid, RegisteredGameVersion.Version, RegisteredGameVersion.GetFriendlyName());
+
+		bInitializedCustomVersions = true;
+	}
+}
+
+void FNetworkVersion::RegisterNetworkCustomVersion(const FGuid& VersionGuid, int32 Version, int32 CompatibleVersion, const FName& FriendlyName)
+{
+	if (const FCustomVersion* ExistingVersion = NetworkCustomVersions.GetVersion(VersionGuid))
+	{
+		if (ExistingVersion->Version != Version)
+		{
+			UE_LOG(LogNetVersion, Warning, TEXT("RegisterNetworkCustomVersion: Overwriting previous version %d with %d for %s"), ExistingVersion->Version, Version, *FriendlyName.ToString());
+		}
+	}
+
+	NetworkCustomVersions.SetVersion(VersionGuid, Version, FriendlyName);
+
+	if (const FCustomVersion* ExistingVersion = CompatibleNetworkCustomVersions.GetVersion(VersionGuid))
+	{
+		if (ExistingVersion->Version != CompatibleVersion)
+		{
+			UE_LOG(LogNetVersion, Warning, TEXT("RegisterNetworkCustomVersion: Overwriting previous compatible version %d with %d for %s"), ExistingVersion->Version, CompatibleVersion, *FriendlyName.ToString());
+		}
+	}
+
+	CompatibleNetworkCustomVersions.SetVersion(VersionGuid, CompatibleVersion, FriendlyName);
+
+	bHasCachedNetworkChecksum = false;
+	bHasCachedReplayChecksum = false;
 }
 
 uint32 FNetworkVersion::GetLocalNetworkVersion( bool AllowOverrideDelegate /*=true*/ )
@@ -144,12 +238,18 @@ uint32 FNetworkVersion::GetLocalNetworkVersion( bool AllowOverrideDelegate /*=tr
 		return CachedNetworkChecksum;
 	}
 
-	FString VersionString = FString::Printf(TEXT("%s %s, NetCL: %d, EngineNetVer: %d, GameNetVer: %d"),
+	InitCustomVersions();
+
+	FString VersionString = FString::Printf(TEXT("%s %s, NetCL: %d"),
 		FApp::GetProjectName(),
 		*FNetworkVersion::GetProjectVersion(),
-		GetNetworkCompatibleChangelist(),
-		FNetworkVersion::GetEngineNetworkProtocolVersion(),
-		FNetworkVersion::GetGameNetworkProtocolVersion());
+		GetNetworkCompatibleChangelist());
+
+	const FCustomVersionArray& CustomVers = NetworkCustomVersions.GetAllVersions();
+	for (const FCustomVersion& CustomVer : CustomVers)
+	{
+		VersionString += FString::Printf(TEXT(", %s: %d"), *CustomVer.GetFriendlyName().ToString(), CustomVer.Version);
+	}
 
 	CachedNetworkChecksum = FCrc::StrCrc32(*VersionString.ToLower());
 
@@ -172,9 +272,22 @@ bool FNetworkVersion::IsNetworkCompatible( const uint32 LocalNetworkVersion, con
 
 FNetworkReplayVersion FNetworkVersion::GetReplayVersion()
 {
-	const uint32 ReplayVersion = ( GameCompatibleNetworkProtocolVersion << 16 ) | EngineCompatibleNetworkProtocolVersion;
+	if (!bHasCachedReplayChecksum)
+	{
+		FString VersionString = FApp::GetProjectName();
+
+		const FCustomVersionArray& CustomVers = CompatibleNetworkCustomVersions.GetAllVersions();
+		for (const FCustomVersion& CustomVer : CustomVers)
+		{
+			VersionString += FString::Printf(TEXT(", %s: %d"), *CustomVer.GetFriendlyName().ToString(), CustomVer.Version);
+		}
+
+		CachedReplayChecksum = FCrc::StrCrc32(*VersionString.ToLower());
+		
+		bHasCachedReplayChecksum = true;
+	}
 
-	return FNetworkReplayVersion( FApp::GetProjectName(), ReplayVersion, GetReplayCompatibleChangelist() );
+	return FNetworkReplayVersion(FApp::GetProjectName(), CachedReplayChecksum, GetReplayCompatibleChangelist());
 }
 
 bool FNetworkVersion::AreNetworkRuntimeFeaturesCompatible(EEngineNetworkRuntimeFeatures LocalFeatures, EEngineNetworkRuntimeFeatures RemoteFeatures)
@@ -192,5 +305,4 @@ void FNetworkVersion::DescribeNetworkRuntimeFeaturesBitset(EEngineNetworkRuntime
 	{
 		OutVerboseDescription.Append(TEXT("LegacyReplication"));
 	}
-}
-
+}
\ No newline at end of file
diff --git a/Engine/Source/Runtime/Core/Private/Serialization/Archive.cpp b/Engine/Source/Runtime/Core/Private/Serialization/Archive.cpp
--- a/Engine/Source/Runtime/Core/Private/Serialization/Archive.cpp
+++ b/Engine/Source/Runtime/Core/Private/Serialization/Archive.cpp
@@ -25,6 +25,7 @@
 #include "UObject/NameTypes.h"
 #include "Compression/CompressionUtil.h"
 #include "UObject/UnrealNames.h"
+#include "Misc/EngineNetworkCustomVersion.h"
 
 PRAGMA_DISABLE_UNSAFE_TYPECAST_WARNINGS
 
@@ -155,8 +156,10 @@ void FArchiveState::Reset()
 	ArUEVer								= GPackageFileUEVersion;
 	ArLicenseeUEVer						= GPackageFileLicenseeUEVersion;
 	ArEngineVer							= FEngineVersion::Current();
-	ArEngineNetVer						= FNetworkVersion::GetEngineNetworkProtocolVersion();
-	ArGameNetVer						= FNetworkVersion::GetGameNetworkProtocolVersion();
+	PRAGMA_DISABLE_DEPRECATION_WARNINGS
+	ArEngineNetVer						= FNetworkVersion::GetNetworkProtocolVersion(FEngineNetworkCustomVersion::Guid);
+	ArGameNetVer						= FNetworkVersion::GetNetworkProtocolVersion(FGameNetworkCustomVersion::Guid);
+	PRAGMA_ENABLE_DEPRECATION_WARNINGS
 	ArIsLoading							= false;
 	ArIsLoadingFromCookedPackage		= false;
 	ArIsSaving							= false;
@@ -215,8 +218,10 @@ void FArchiveState::CopyTrivialFArchiveStatusMembers(const FArchiveState& Archiv
 	ArUEVer                              = ArchiveToCopy.ArUEVer;
 	ArLicenseeUEVer                      = ArchiveToCopy.ArLicenseeUEVer;
 	ArEngineVer                          = ArchiveToCopy.ArEngineVer;
+	PRAGMA_DISABLE_DEPRECATION_WARNINGS
 	ArEngineNetVer                       = ArchiveToCopy.ArEngineNetVer;
 	ArGameNetVer                         = ArchiveToCopy.ArGameNetVer;
+	PRAGMA_ENABLE_DEPRECATION_WARNINGS
 	ArIsLoading                          = ArchiveToCopy.ArIsLoading;
 	ArIsLoadingFromCookedPackage         = ArchiveToCopy.ArIsLoadingFromCookedPackage;
 	ArIsSaving                           = ArchiveToCopy.ArIsSaving;
@@ -1373,12 +1378,28 @@ void FArchiveState::SetEngineVer(const FEngineVersionBase& InVer)
 
 void FArchiveState::SetEngineNetVer(const uint32 InEngineNetVer)
 {
+	PRAGMA_DISABLE_DEPRECATION_WARNINGS
 	ArEngineNetVer = InEngineNetVer;
+	PRAGMA_ENABLE_DEPRECATION_WARNINGS
+	SetCustomVersion(FEngineNetworkCustomVersion::Guid, InEngineNetVer, TEXT("EngineNetworkVersion"));
+}
+
+uint32 FArchiveState::EngineNetVer() const
+{
+	return CustomVer(FEngineNetworkCustomVersion::Guid);
 }
 
 void FArchiveState::SetGameNetVer(const uint32 InGameNetVer)
 {
+	PRAGMA_DISABLE_DEPRECATION_WARNINGS
 	ArGameNetVer = InGameNetVer;
+	PRAGMA_ENABLE_DEPRECATION_WARNINGS
+	SetCustomVersion(FGameNetworkCustomVersion::Guid, InGameNetVer, TEXT("GameNetworkVersion"));
+}
+
+uint32 FArchiveState::GameNetVer() const
+{
+	return CustomVer(FGameNetworkCustomVersion::Guid);
 }
 
 void FArchiveState::SetIsLoading(bool bInIsLoading)
diff --git a/Engine/Source/Runtime/Core/Public/Math/Vector.h b/Engine/Source/Runtime/Core/Public/Math/Vector.h
--- a/Engine/Source/Runtime/Core/Public/Math/Vector.h
+++ b/Engine/Source/Runtime/Core/Public/Math/Vector.h
@@ -1160,7 +1160,9 @@ public:
      */
 	bool NetSerialize(FArchive& Ar, class UPackageMap* Map, bool& bOutSuccess)
 	{
-		if (Ar.EngineNetVer() >= HISTORY_SERIALIZE_DOUBLE_VECTORS_AS_DOUBLES && Ar.EngineNetVer() != HISTORY_21_AND_VIEWPITCH_ONLY_DO_NOT_USE)
+		Ar.UsingCustomVersion(FEngineNetworkCustomVersion::Guid);
+
+		if (Ar.EngineNetVer() >= FEngineNetworkCustomVersion::SerializeDoubleVectorsAsDoubles && Ar.EngineNetVer() != FEngineNetworkCustomVersion::Ver21AndViewPitchOnly_DONOTUSE)
 		{
 			Ar << X << Y << Z;
 		}
diff --git a/Engine/Source/Runtime/Core/Public/Math/Vector2D.h b/Engine/Source/Runtime/Core/Public/Math/Vector2D.h
--- a/Engine/Source/Runtime/Core/Public/Math/Vector2D.h
+++ b/Engine/Source/Runtime/Core/Public/Math/Vector2D.h
@@ -678,7 +678,9 @@ public:
 	*/
 	bool NetSerialize(FArchive& Ar, class UPackageMap* Map, bool& bOutSuccess)
 	{
-		if (Ar.EngineNetVer() >= HISTORY_SERIALIZE_DOUBLE_VECTORS_AS_DOUBLES && Ar.EngineNetVer() != HISTORY_21_AND_VIEWPITCH_ONLY_DO_NOT_USE)
+		Ar.UsingCustomVersion(FEngineNetworkCustomVersion::Guid);
+
+		if (Ar.EngineNetVer() >= FEngineNetworkCustomVersion::SerializeDoubleVectorsAsDoubles && Ar.EngineNetVer() != FEngineNetworkCustomVersion::Ver21AndViewPitchOnly_DONOTUSE)
 		{
 			Ar << X << Y;
 		}
diff --git a/Engine/Source/Runtime/Core/Public/Misc/EngineNetworkCustomVersion.h b/Engine/Source/Runtime/Core/Public/Misc/EngineNetworkCustomVersion.h
new file mode 100644
--- /dev/null
+++ b/Engine/Source/Runtime/Core/Public/Misc/EngineNetworkCustomVersion.h
@@ -0,0 +1,117 @@
+// Copyright Epic Games, Inc. All Rights Reserved.
+
+#pragma once
+
+#include "CoreTypes.h"
+#include "Misc/Guid.h"
+
+enum UE_DEPRECATED(5.1, "Using custom versions instead going forward, see FEngineNetworkCustomVersion") EEngineNetworkVersionHistory
+{
+	HISTORY_INITIAL = 1,
+	HISTORY_REPLAY_BACKWARDS_COMPAT = 2,			// Bump version to get rid of older replays before backwards compat was turned on officially
+	HISTORY_MAX_ACTOR_CHANNELS_CUSTOMIZATION = 3,	// Bump version because serialization of the actor channels changed
+	HISTORY_REPCMD_CHECKSUM_REMOVE_PRINTF = 4,		// Bump version since the way FRepLayoutCmd::CompatibleChecksum was calculated changed due to an optimization
+	HISTORY_NEW_ACTOR_OVERRIDE_LEVEL = 5,			// Bump version since a level reference was added to the new actor information
+	HISTORY_CHANNEL_NAMES = 6,						// Bump version since channel type is now an fname
+	HISTORY_CHANNEL_CLOSE_REASON = 7,				// Bump version to serialize a channel close reason in bunches instead of bDormant
+	HISTORY_ACKS_INCLUDED_IN_HEADER = 8,			// Bump version since acks are now sent as part of the header
+	HISTORY_NETEXPORT_SERIALIZATION = 9,			// Bump version due to serialization change to FNetFieldExport
+	HISTORY_NETEXPORT_SERIALIZE_FIX = 10,			// Bump version to fix net field export name serialization 
+	HISTORY_FAST_ARRAY_DELTA_STRUCT = 11,			// Bump version to allow fast array serialization, delta struct serialization.
+	HISTORY_FIX_ENUM_SERIALIZATION = 12,			// Bump version to fix enum net serialization issues.
+	HISTORY_OPTIONALLY_QUANTIZE_SPAWN_INFO = 13,	// Bump version to conditionally disable quantization for Scale, Location, and Velocity when spawning network actors.
+	HISTORY_JITTER_IN_HEADER = 14,					// Bump version since we added jitter clock time to packet headers and removed remote saturation
+	HISTORY_CLASSNETCACHE_FULLNAME = 15,			// Bump version to use full paths in GetNetFieldExportGroupForClassNetCache
+	HISTORY_REPLAY_DORMANCY = 16,					// Bump version to support dormancy properly in replays
+	HISTORY_ENUM_SERIALIZATION_COMPAT = 17,			// Bump version to include enum bits required for serialization into compat checksums, as well as unify enum and byte property enum serialization
+	HISTORY_SUBOBJECT_OUTER_CHAIN = 18,				// Bump version to support subobject outer chains matching on client and server
+	HISTORY_HITRESULT_INSTANCEHANDLE = 19,			// Bump version to support FHitResult change of Actor to HitObjectHandle. This change was made in CL 14369221 but a net version wasn't added at the time.
+	HISTORY_INTERFACE_PROPERTY_SERIALIZATION = 20,	// Bump version to support net serialization of FInterfaceProperty
+	HISTORY_MONTAGE_PLAY_INST_ID_SERIALIZATION = 21,// Bump version to support net serialization of FGameplayAbilityRepAnimMontage, addition of PlayInstanceId and removal of bForcePlayBit
+	HISTORY_SERIALIZE_DOUBLE_VECTORS_AS_DOUBLES = 22,// Bump version to support net serialization of double vector types
+	HISTORY_PACKED_VECTOR_LWC_SUPPORT = 23,			// Bump version to support quantized LWC FVector net serialization
+	HISTORY_PAWN_REMOTEVIEWPITCH = 24,				// Bump version to support serialization changes to RemoteViewPitch
+	HISTORY_REPMOVE_SERVERFRAME_AND_HANDLE = 25,	// Bump version to support serialization changes to RepMove so we can get the serverframe and physics handle associated with the object
+	HISTORY_21_AND_VIEWPITCH_ONLY_DO_NOT_USE = 26,  // Bump version to support up to history 21 + HISTORY_PAWN_REMOTEVIEWPITCH.  DO NOT USE!!!
+	HISTORY_PLACEHOLDER = 27,                       // Bump version to a placeholder.  This version is the same as HISTORY_REPMOVE_SERVERFRAME_AND_HANDLE
+	HISTORY_RUNTIME_FEATURES_COMPATIBILITY = 28,	// Bump version to add network runtime feature compatibility test to handshake (hello/upgrade) control messages
+	HISTORY_SOFTOBJECTPTR_NETGUIDS = 29,			// Bump version to support replicating SoftObjectPtrs by NetGuid instead of raw strings.
+	HISTORY_SUBOBJECT_DESTROY_FLAG = 30,			// Bump version to support subobject destruction message flags
+	HISTORY_GAMESTATE_REPLCIATED_TIME_AS_DOUBLE = 31,	// Bump version to support AGameStateBase::ReplicatedWorldTimeSeconds as double instead of float.
+	HISTORY_CUSTOMVERION = 32,                      // Bump version to switch to using custom versions
+	// New history items go above here.
+
+	HISTORY_ENGINENETVERSION_PLUS_ONE,
+	HISTORY_ENGINENETVERSION_LATEST = HISTORY_ENGINENETVERSION_PLUS_ONE - 1,
+};
+
+struct CORE_API FEngineNetworkCustomVersion
+{
+	enum Type
+	{
+		// Before any version changes were made
+		BeforeCustomVersionWasAdded = 0,
+
+		// Original replay versions from ENetworkVersionHistory
+		Initial = 1,
+		ReplayBackwardsCompat = 2,				// Bump version to get rid of older replays before backwards compat was turned on officially
+		MaxActorChannelsCustomization = 3,		// Bump version because serialization of the actor channels changed
+		RepCmdChecksumRemovePrintf = 4,			// Bump version since the way FRepLayoutCmd::CompatibleChecksum was calculated changed due to an optimization
+		NewActorOverrideLevel = 5,				// Bump version since a level reference was added to the new actor information
+		ChannelNames = 6,						// Bump version since channel type is now an fname
+		ChannelCloseReason = 7,					// Bump version to serialize a channel close reason in bunches instead of bDormant
+		AcksIncludedInHeader = 8,				// Bump version since acks are now sent as part of the header
+		NetExportSerialization = 9,				// Bump version due to serialization change to FNetFieldExport
+		NetExportSerializeFix = 10,				// Bump version to fix net field export name serialization 
+		FastArrayDeltaStruct = 11,				// Bump version to allow fast array serialization, delta struct serialization.
+		FixEnumSerialization = 12,				// Bump version to fix enum net serialization issues.
+		OptionallyQuantizeSpawnInfo = 13,		// Bump version to conditionally disable quantization for Scale, Location, and Velocity when spawning network actors.
+		JitterInHeader = 14,					// Bump version since we added jitter clock time to packet headers and removed remote saturation
+		ClassNetCacheFullName = 15,				// Bump version to use full paths in GetNetFieldExportGroupForClassNetCache
+		ReplayDormancy = 16,					// Bump version to support dormancy properly in replays
+		EnumSerializationCompat = 17,			// Bump version to include enum bits required for serialization into compat checksums, as well as unify enum and byte property enum serialization
+		SubObjectOuterChain = 18,				// Bump version to support subobject outer chains matching on client and server
+		HitResultInstanceHandle = 19,			// Bump version to support FHitResult change of Actor to HitObjectHandle. This change was made in CL 14369221 but a net version wasn't added at the time.
+		InterfacePropertySerialization = 20,	// Bump version to support net serialization of FInterfaceProperty
+		MontagePlayInstIdSerialization = 21,	// Bump version to support net serialization of FGameplayAbilityRepAnimMontage, addition of PlayInstanceId and removal of bForcePlayBit
+		SerializeDoubleVectorsAsDoubles = 22,	// Bump version to support net serialization of double vector types
+		PackedVectorLWCSupport = 23,			// Bump version to support quantized LWC FVector net serialization
+		PawnRemoteViewPitch = 24,				// Bump version to support serialization changes to RemoteViewPitch
+		RepMoveServerFrameAndHandle = 25,		// Bump version to support serialization changes to RepMove so we can get the serverframe and physics handle associated with the object
+		Ver21AndViewPitchOnly_DONOTUSE = 26,	// Bump version to support up to history 21 + PawnRemoteViewPitch.  DO NOT USE!!!
+		Placeholder = 27,                       // Bump version to a placeholder.  This version is the same as RepMoveServerFrameAndHandle
+		RuntimeFeaturesCompatibility = 28,		// Bump version to add network runtime feature compatibility test to handshake (hello/upgrade) control messages
+		SoftObjectPtrNetGuids = 29,				// Bump version to support replicating SoftObjectPtrs by NetGuid instead of raw strings.
+		SubObjectDestroyFlag = 30,				// Bump version to support subobject destruction message flags
+		GameStateReplicatedTimeAsDouble = 31,	// Bump version to support AGameStateBase::ReplicatedWorldTimeSeconds as double instead of float.
+
+		CustomVersions = 32,                    // Bump version to switch to using custom versions
+
+		// -----<new versions can be added above this line>-------------------------------------------------
+		VersionPlusOne,
+		LatestVersion = VersionPlusOne - 1
+	};
+
+	// The GUID for this custom version number
+	const static FGuid Guid;
+
+	FEngineNetworkCustomVersion() = delete;
+};
+
+struct CORE_API FGameNetworkCustomVersion
+{
+	enum Type
+	{
+		// Before any version changes were made
+		BeforeCustomVersionWasAdded = 0,
+
+		// -----<new versions can be added above this line>-------------------------------------------------
+		VersionPlusOne,
+		LatestVersion = VersionPlusOne - 1
+	};
+
+	// The GUID for this custom version number
+	const static FGuid Guid;
+
+	FGameNetworkCustomVersion() = delete;
+};
\ No newline at end of file
diff --git a/Engine/Source/Runtime/Core/Public/Misc/NetworkVersion.h b/Engine/Source/Runtime/Core/Public/Misc/NetworkVersion.h
--- a/Engine/Source/Runtime/Core/Public/Misc/NetworkVersion.h
+++ b/Engine/Source/Runtime/Core/Public/Misc/NetworkVersion.h
@@ -6,6 +6,7 @@
 #include "CoreTypes.h"
 #include "Delegates/Delegate.h"
 #include "Logging/LogMacros.h"
+#include "Misc/EngineNetworkCustomVersion.h"
 
 // The version number used for determining network compatibility. If zero, uses the engine compatible version.
 #define ENGINE_NET_VERSION  0
@@ -18,14 +19,21 @@ CORE_API DECLARE_LOG_CATEGORY_EXTERN( LogNetVersion, Log, All );
 class FNetworkReplayVersion
 {
 public:
-	FNetworkReplayVersion() : NetworkVersion( 0 ), Changelist( 0 )
+	FNetworkReplayVersion() 
+		: NetworkVersion(0)
+		, Changelist(0)
 	{
 	}
-	FNetworkReplayVersion( const FString& InAppString, const uint32 InNetworkVersion, const uint32 InChangelist ) : AppString( InAppString ), NetworkVersion( InNetworkVersion ), Changelist( InChangelist )
+
+	FNetworkReplayVersion(const FString& InAppString, const uint32 InNetworkVersion, const uint32 InChangelist) 
+		: AppString(InAppString)
+		, NetworkVersion(InNetworkVersion)
+		, Changelist(InChangelist)
 	{
 	}
 
 	FString		AppString;
+	/** This is a hash of compatible versions, and not pulled directly from any version enums. */
 	uint32		NetworkVersion;
 	uint32		Changelist;
 };
@@ -40,45 +48,6 @@ enum class EEngineNetworkRuntimeFeatures : uint16
 };
 ENUM_CLASS_FLAGS(EEngineNetworkRuntimeFeatures);
 
-enum EEngineNetworkVersionHistory
-{
-	HISTORY_INITIAL = 1,
-	HISTORY_REPLAY_BACKWARDS_COMPAT = 2,				// Bump version to get rid of older replays before backwards compat was turned on officially
-	HISTORY_MAX_ACTOR_CHANNELS_CUSTOMIZATION = 3,		// Bump version because serialization of the actor channels changed
-	HISTORY_REPCMD_CHECKSUM_REMOVE_PRINTF = 4,			// Bump version since the way FRepLayoutCmd::CompatibleChecksum was calculated changed due to an optimization
-	HISTORY_NEW_ACTOR_OVERRIDE_LEVEL = 5,				// Bump version since a level reference was added to the new actor information
-	HISTORY_CHANNEL_NAMES = 6,							// Bump version since channel type is now an fname
-	HISTORY_CHANNEL_CLOSE_REASON = 7,					// Bump version to serialize a channel close reason in bunches instead of bDormant
-	HISTORY_ACKS_INCLUDED_IN_HEADER = 8,				// Bump version since acks are now sent as part of the header
-	HISTORY_NETEXPORT_SERIALIZATION = 9,				// Bump version due to serialization change to FNetFieldExport
-	HISTORY_NETEXPORT_SERIALIZE_FIX = 10,				// Bump version to fix net field export name serialization 
-	HISTORY_FAST_ARRAY_DELTA_STRUCT = 11,				// Bump version to allow fast array serialization, delta struct serialization.
-	HISTORY_FIX_ENUM_SERIALIZATION = 12,				// Bump version to fix enum net serialization issues.
-	HISTORY_OPTIONALLY_QUANTIZE_SPAWN_INFO = 13,		// Bump version to conditionally disable quantization for Scale, Location, and Velocity when spawning network actors.
-	HISTORY_JITTER_IN_HEADER = 14,						// Bump version since we added jitter clock time to packet headers and removed remote saturation
-	HISTORY_CLASSNETCACHE_FULLNAME = 15,				// Bump version to use full paths in GetNetFieldExportGroupForClassNetCache
-	HISTORY_REPLAY_DORMANCY = 16,						// Bump version to support dormancy properly in replays
-	HISTORY_ENUM_SERIALIZATION_COMPAT = 17,				// Bump version to include enum bits required for serialization into compat checksums, as well as unify enum and byte property enum serialization
-	HISTORY_SUBOBJECT_OUTER_CHAIN = 18,					// Bump version to support subobject outer chains matching on client and server
-	HISTORY_HITRESULT_INSTANCEHANDLE = 19,				// Bump version to support FHitResult change of Actor to HitObjectHandle. This change was made in CL 14369221 but a net version wasn't added at the time.
-	HISTORY_INTERFACE_PROPERTY_SERIALIZATION = 20,		// Bump version to support net serialization of FInterfaceProperty
-	HISTORY_MONTAGE_PLAY_INST_ID_SERIALIZATION = 21,	// Bump version to support net serialization of FGameplayAbilityRepAnimMontage, addition of PlayInstanceId and removal of bForcePlayBit
-	HISTORY_SERIALIZE_DOUBLE_VECTORS_AS_DOUBLES	= 22,	// Bump version to support net serialization of double vector types
-	HISTORY_PACKED_VECTOR_LWC_SUPPORT = 23,				// Bump version to support quantized LWC FVector net serialization
-	HISTORY_PAWN_REMOTEVIEWPITCH = 24,					// Bump version to support serialization changes to RemoteViewPitch
-	HISTORY_REPMOVE_SERVERFRAME_AND_HANDLE = 25,		// Bump version to support serialization changes to RepMove so we can get the serverframe and physics handle associated with the object
-	HISTORY_21_AND_VIEWPITCH_ONLY_DO_NOT_USE = 26,		// Bump version to support up to history 21 + HISTORY_PAWN_REMOTEVIEWPITCH.  DO NOT USE!!!
-	HISTORY_PLACEHOLDER = 27,							// Bump version to a placeholder.  This version is the same as HISTORY_REPMOVE_SERVERFRAME_AND_HANDLE
-	HISTORY_RUNTIME_FEATURES_COMPATIBILITY = 28,		// Bump version to add network runtime feature compatibility test to handshake (hello/upgrade) control messages
-	HISTORY_SOFTOBJECTPTR_NETGUIDS = 29,				// Bump version to support replicating SoftObjectPtrs by NetGuid instead of raw strings.
-	HISTORY_SUBOBJECT_DESTROY_FLAG = 30,				// Bump version to support subobject destruction message flags
-	HISTORY_GAMESTATE_REPLCIATED_TIME_AS_DOUBLE = 31,	// Bump version to support AGameStateBase::ReplicatedWorldTimeSeconds as double instead of float.
-	// New history items go above here.
-
-	HISTORY_ENGINENETVERSION_PLUS_ONE,
-	HISTORY_ENGINENETVERSION_LATEST = HISTORY_ENGINENETVERSION_PLUS_ONE - 1,
-};
-
 struct CORE_API FNetworkVersion
 {
 	/** Called in GetLocalNetworkVersion if bound */
@@ -95,11 +64,24 @@ struct CORE_API FNetworkVersion
 
 	static uint32 GetNetworkCompatibleChangelist();
 	static uint32 GetReplayCompatibleChangelist();
+
+	UE_DEPRECATED(5.2, "Please use GetCompatibleNetworkProtocolVersion instead.")
 	static uint32 GetEngineNetworkProtocolVersion();
+	UE_DEPRECATED(5.2, "Please use GetCompatibleNetworkProtocolVersion instead.")
 	static uint32 GetGameNetworkProtocolVersion();
+
+	UE_DEPRECATED(5.2, "Please use GetCompatibleNetworkProtocolVersion instead.")
 	static uint32 GetEngineCompatibleNetworkProtocolVersion();
+	UE_DEPRECATED(5.2, "Please use GetCompatibleNetworkProtocolVersion instead.")
 	static uint32 GetGameCompatibleNetworkProtocolVersion();
 
+	static uint32 GetNetworkProtocolVersion(const FGuid& VersionGuid);
+	static uint32 GetCompatibleNetworkProtocolVersion(const FGuid& VersionGuid);
+
+	static const FCustomVersionContainer& GetNetworkCustomVersions();
+
+	static void RegisterNetworkCustomVersion(const FGuid& VersionGuid, int32 Version, int32 CompatibleVersion, const FName& FriendlyName);
+
 	/**
 	* Generates a version number, that by default, is based on a checksum of the engine version + project name + project version string
 	* Game/project code can completely override what this value returns through the GetLocalNetworkVersionOverride delegate
@@ -174,9 +156,23 @@ protected:
 	static bool		bHasCachedNetworkChecksum;
 	static uint32	CachedNetworkChecksum;
 
+	static bool		bHasCachedReplayChecksum;
+	static uint32	CachedReplayChecksum;
+
+	UE_DEPRECATED(5.2, "Now storing this value in NetworkCustomVersions, do not use directly.")
 	static uint32	EngineNetworkProtocolVersion;
+	UE_DEPRECATED(5.2, "Now storing this value in NetworkCustomVersions, do not use directly.")
 	static uint32	GameNetworkProtocolVersion;
 
+	UE_DEPRECATED(5.2, "Now storing this value in CompatibleNetworkCustomVersions, do not use directly.")
 	static uint32	EngineCompatibleNetworkProtocolVersion;
+	UE_DEPRECATED(5.2, "Now storing this value in CompatibleNetworkCustomVersions, do not use directly.")
 	static uint32	GameCompatibleNetworkProtocolVersion;
+
+private:
+	static FCustomVersionContainer NetworkCustomVersions;
+	static FCustomVersionContainer CompatibleNetworkCustomVersions;
+
+	static bool bInitializedCustomVersions;
+	static void InitCustomVersions();
 };
diff --git a/Engine/Source/Runtime/Core/Public/Serialization/Archive.h b/Engine/Source/Runtime/Core/Public/Serialization/Archive.h
--- a/Engine/Source/Runtime/Core/Public/Serialization/Archive.h
+++ b/Engine/Source/Runtime/Core/Public/Serialization/Archive.h
@@ -229,16 +229,10 @@ public:
 	}
 
 	/** Returns the engine-global network protocol version for this archive. */
-	FORCEINLINE uint32 EngineNetVer() const
-	{
-		return ArEngineNetVer;
-	}
+	uint32 EngineNetVer() const;
 
 	/** Returns the game-specific network protocol version for this archive. */
-	FORCEINLINE uint32 GameNetVer() const
-	{
-		return ArGameNetVer;
-	}
+	uint32 GameNetVer() const;
 
 	/**
 	 * Queries a custom version from the archive.  If the archive is being used to write, the custom version must have already been registered.
@@ -973,9 +967,11 @@ protected:
 	FEngineVersionBase ArEngineVer;
 
 	/** Holds the engine network protocol version. */
+	UE_DEPRECATED(5.2, "Engine net version moved to custom versions")
 	uint32 ArEngineNetVer;
 
 	/** Holds the game network protocol version. */
+	UE_DEPRECATED(5.2, "Game net version moved to custom versions")
 	uint32 ArGameNetVer;
 
 	/**
diff --git a/Engine/Source/Runtime/CoreUObject/Private/UObject/CoreNet.cpp b/Engine/Source/Runtime/CoreUObject/Private/UObject/CoreNet.cpp
--- a/Engine/Source/Runtime/CoreUObject/Private/UObject/CoreNet.cpp
+++ b/Engine/Source/Runtime/CoreUObject/Private/UObject/CoreNet.cpp
@@ -6,7 +6,7 @@
 
 #include "UObject/CoreNet.h"
 #include "UObject/UnrealType.h"
-#include "Misc/NetworkVersion.h"
+#include "Misc/EngineNetworkCustomVersion.h"
 
 DEFINE_LOG_CATEGORY_STATIC(LogCoreNet, Log, All);
 
@@ -294,7 +294,7 @@ bool UPackageMap::StaticSerializeName(FArchive& Ar, FName& InName)
 		{
 			// replicated by hardcoded index
 			uint32 NameIndex;
-			if (Ar.EngineNetVer() < HISTORY_CHANNEL_NAMES)
+			if (Ar.EngineNetVer() < FEngineNetworkCustomVersion::ChannelNames)
 			{
 				Ar.SerializeInt(NameIndex, MAX_NETWORKED_HARDCODED_NAME + 1);
 			}
@@ -445,8 +445,9 @@ FArchive& FNetBitWriter::operator<<(FSoftObjectPath& Value)
 FArchive& FNetBitWriter::operator<<(FSoftObjectPtr& Value)
 {
 	// Keep in sync with FNetBitReader::operator<<(FSoftObjectPtr&)
+	UsingCustomVersion(FEngineNetworkCustomVersion::Guid);
 
-	if (EngineNetVer() >= EEngineNetworkVersionHistory::HISTORY_SOFTOBJECTPTR_NETGUIDS)
+	if (EngineNetVer() >= FEngineNetworkCustomVersion::SoftObjectPtrNetGuids)
 	{
 		// Treat soft pointers to dynamic (non-stably named) objects as weak pointers.
 		// This allows these objects to be resolved on clients and to only replicate a NetGUID.
@@ -536,8 +537,9 @@ FArchive& FNetBitReader::operator<<(FSoftObjectPath& Value)
 FArchive& FNetBitReader::operator<<(FSoftObjectPtr& Value)
 {
 	// Keep in sync with FNetBitWriter::operator<<(FSoftObjectPtr&)
+	UsingCustomVersion(FEngineNetworkCustomVersion::Guid);
 
-	if (EngineNetVer() >= EEngineNetworkVersionHistory::HISTORY_SOFTOBJECTPTR_NETGUIDS)
+	if (EngineNetVer() >= FEngineNetworkCustomVersion::SoftObjectPtrNetGuids)
 	{
 		// Treat soft pointers to dynamic (non-stably named) objects as weak pointers.
 		// This allows these objects to be resolved on clients and to only replicate a NetGUID.
diff --git a/Engine/Source/Runtime/CoreUObject/Private/UObject/EnumProperty.cpp b/Engine/Source/Runtime/CoreUObject/Private/UObject/EnumProperty.cpp
--- a/Engine/Source/Runtime/CoreUObject/Private/UObject/EnumProperty.cpp
+++ b/Engine/Source/Runtime/CoreUObject/Private/UObject/EnumProperty.cpp
@@ -9,7 +9,7 @@
 #include "Templates/IsSigned.h"
 #include "Algo/Find.h"
 #include "UObject/LinkerLoad.h"
-#include "Misc/NetworkVersion.h"
+#include "Misc/EngineNetworkCustomVersion.h"
 #include "Hash/Blake3.h"
 
 namespace UEEnumProperty_Private
@@ -198,7 +198,9 @@ void FEnumProperty::SerializeItem(FStructuredArchive::FSlot Slot, void* Value, v
 
 bool FEnumProperty::NetSerializeItem(FArchive& Ar, UPackageMap* Map, void* Data, TArray<uint8>* MetaData) const
 {
-	if (Ar.EngineNetVer() < HISTORY_FIX_ENUM_SERIALIZATION)
+	Ar.UsingCustomVersion(FEngineNetworkCustomVersion::Guid);
+
+	if (Ar.EngineNetVer() < FEngineNetworkCustomVersion::FixEnumSerialization)
 	{
 		Ar.SerializeBits(Data, FMath::CeilLogTwo64(Enum->GetMaxEnumValue()));
 	}
diff --git a/Engine/Source/Runtime/CoreUObject/Private/UObject/PropertyByte.cpp b/Engine/Source/Runtime/CoreUObject/Private/UObject/PropertyByte.cpp
--- a/Engine/Source/Runtime/CoreUObject/Private/UObject/PropertyByte.cpp
+++ b/Engine/Source/Runtime/CoreUObject/Private/UObject/PropertyByte.cpp
@@ -10,7 +10,7 @@
 #include "Serialization/ArchiveUObjectFromStructuredArchive.h"
 #include "Algo/Find.h"
 #include "UObject/LinkerLoad.h"
-#include "Misc/NetworkVersion.h"
+#include "Misc/EngineNetworkCustomVersion.h"
 #include "Hash/Blake3.h"
 
 /*-----------------------------------------------------------------------------
@@ -102,7 +102,9 @@ void FByteProperty::SerializeItem( FStructuredArchive::FSlot Slot, void* Value,
 }
 bool FByteProperty::NetSerializeItem( FArchive& Ar, UPackageMap* Map, void* Data, TArray<uint8> * MetaData ) const
 {
-	if (Ar.EngineNetVer() < HISTORY_ENUM_SERIALIZATION_COMPAT)
+	Ar.UsingCustomVersion(FEngineNetworkCustomVersion::Guid);
+
+	if (Ar.EngineNetVer() < FEngineNetworkCustomVersion::EnumSerializationCompat)
 	{
 		Ar.SerializeBits(Data, Enum ? FMath::CeilLogTwo(Enum->GetMaxEnumValue()) : 8);
 	}
diff --git a/Engine/Source/Runtime/CoreUObject/Private/UObject/PropertyInterface.cpp b/Engine/Source/Runtime/CoreUObject/Private/UObject/PropertyInterface.cpp
--- a/Engine/Source/Runtime/CoreUObject/Private/UObject/PropertyInterface.cpp
+++ b/Engine/Source/Runtime/CoreUObject/Private/UObject/PropertyInterface.cpp
@@ -11,7 +11,7 @@
 #include "UObject/LinkerPlaceholderClass.h"
 #include "Hash/Blake3.h"
 #include "UObject/CoreNet.h"
-#include "Misc/NetworkVersion.h"
+#include "Misc/EngineNetworkCustomVersion.h"
 
 /*-----------------------------------------------------------------------------
 	FInterfaceProperty.
@@ -159,7 +159,9 @@ void FInterfaceProperty::SerializeItem( FStructuredArchive::FSlot Slot, void* Va
 
 bool FInterfaceProperty::NetSerializeItem( FArchive& Ar, UPackageMap* Map, void* Data, TArray<uint8> * MetaData ) const
 {
-	if (Ar.EngineNetVer() >= HISTORY_INTERFACE_PROPERTY_SERIALIZATION)
+	Ar.UsingCustomVersion(FEngineNetworkCustomVersion::Guid);
+
+	if (Ar.EngineNetVer() >= FEngineNetworkCustomVersion::InterfacePropertySerialization)
 	{ 
 		FScriptInterface* InterfaceValue = (FScriptInterface*)Data;
 		bool Result = Map->SerializeObject(Ar, InterfaceClass, InterfaceValue->GetObjectRef());
diff --git a/Engine/Source/Runtime/Engine/Classes/Engine/DemoNetDriver.h b/Engine/Source/Runtime/Engine/Classes/Engine/DemoNetDriver.h
--- a/Engine/Source/Runtime/Engine/Classes/Engine/DemoNetDriver.h
+++ b/Engine/Source/Runtime/Engine/Classes/Engine/DemoNetDriver.h
@@ -627,10 +627,21 @@ public:
 	}
 
 	uint32 GetPlaybackDemoChangelist() const { return ReplayHelper.PlaybackDemoHeader.EngineVersion.GetChangelist(); }
-	uint32 GetPlaybackDemoVersion() const { return ReplayHelper.PlaybackDemoHeader.Version; }
 
-	uint32 GetPlaybackEngineNetworkProtocolVersion() const { return ReplayHelper.PlaybackDemoHeader.EngineNetworkProtocolVersion; }
-	uint32 GetPlaybackGameNetworkProtocolVersion() const { return ReplayHelper.PlaybackDemoHeader.GameNetworkProtocolVersion; }
+	UE_DEPRECATED(5.2, "Will be removed in favor of custom versions, use GetPlaybackReplayVersion instead")
+	uint32 GetPlaybackDemoVersion() const 
+	{ 
+		PRAGMA_DISABLE_DEPRECATION_WARNINGS
+		return ReplayHelper.PlaybackDemoHeader.Version; 
+		PRAGMA_ENABLE_DEPRECATION_WARNINGS
+	}
+
+	FReplayCustomVersion::Type GetPlaybackReplayVersion() const { return ReplayHelper.GetPlaybackReplayVersion(); }
+
+	uint32 GetPlaybackEngineNetworkProtocolVersion() const { return ReplayHelper.PlaybackDemoHeader.GetCustomVersion(FEngineNetworkCustomVersion::Guid); }
+	uint32 GetPlaybackGameNetworkProtocolVersion() const { return ReplayHelper.PlaybackDemoHeader.GetCustomVersion(FGameNetworkCustomVersion::Guid); }
+
+	uint32 GetPlaybackCustomVersionVersion(const FGuid& VersionGuid) const { return ReplayHelper.PlaybackDemoHeader.GetCustomVersion(VersionGuid); }
 
 	FString GetDemoPath() const;
 
diff --git a/Engine/Source/Runtime/Engine/Classes/Engine/NetConnection.h b/Engine/Source/Runtime/Engine/Classes/Engine/NetConnection.h
--- a/Engine/Source/Runtime/Engine/Classes/Engine/NetConnection.h
+++ b/Engine/Source/Runtime/Engine/Classes/Engine/NetConnection.h
@@ -15,6 +15,7 @@
 #include "UObject/ObjectMacros.h"
 #include "UObject/UObjectGlobals.h"
 #include "Serialization/BitWriter.h"
+#include "Serialization/CustomVersion.h"
 #include "Misc/EngineVersion.h"
 #include "Misc/NetworkGuid.h"
 #include "GameFramework/OnlineReplStructs.h"
@@ -604,9 +605,19 @@ public:
 	int32				InitInReliable;
 
 	// Network version
+	UE_DEPRECATED(5.2, "Deprecated in favor of NetworkCustomVersions, please use GetNetworkCustomVersion instead")
 	uint32				EngineNetworkProtocolVersion;
+	UE_DEPRECATED(5.2, "Deprecated in favor of NetworkCustomVersions, please use GetNetworkCustomVersion instead")
 	uint32				GameNetworkProtocolVersion;
 
+	uint32 GetNetworkCustomVersion(const FGuid& VersionGuid) const;
+	void SetNetworkCustomVersions(const FCustomVersionContainer& CustomVersions);
+
+private:
+	FCustomVersionContainer NetworkCustomVersions;
+
+public:
+
 	// Log tracking
 	double			LogCallLastTime;
 	int32			LogCallCount;
diff --git a/Engine/Source/Runtime/Engine/Classes/Engine/NetSerialization.h b/Engine/Source/Runtime/Engine/Classes/Engine/NetSerialization.h
--- a/Engine/Source/Runtime/Engine/Classes/Engine/NetSerialization.h
+++ b/Engine/Source/Runtime/Engine/Classes/Engine/NetSerialization.h
@@ -221,7 +221,7 @@ bool WritePackedVector(FVector3d Vector, FArchive& Ar)
 template<int32 ScaleFactor, int32 MaxBitsPerComponent>
 bool ReadPackedVector(FVector3f& Value, FArchive& Ar)
 {
-	if (Ar.EngineNetVer() >= HISTORY_PACKED_VECTOR_LWC_SUPPORT && Ar.EngineNetVer() != HISTORY_21_AND_VIEWPITCH_ONLY_DO_NOT_USE)
+	if (Ar.EngineNetVer() >= FEngineNetworkCustomVersion::PackedVectorLWCSupport && Ar.EngineNetVer() != FEngineNetworkCustomVersion::Ver21AndViewPitchOnly_DONOTUSE)
 	{
 		return UE::Net::ReadQuantizedVector(ScaleFactor, Value, Ar);
 	}
@@ -234,7 +234,7 @@ bool ReadPackedVector(FVector3f& Value, FArchive& Ar)
 template<int32 ScaleFactor, int32 MaxBitsPerComponent>
 bool ReadPackedVector(FVector3d& Value, FArchive& Ar)
 {
-	if (Ar.EngineNetVer() >= HISTORY_PACKED_VECTOR_LWC_SUPPORT && Ar.EngineNetVer() != HISTORY_21_AND_VIEWPITCH_ONLY_DO_NOT_USE)
+	if (Ar.EngineNetVer() >= FEngineNetworkCustomVersion::PackedVectorLWCSupport && Ar.EngineNetVer() != FEngineNetworkCustomVersion::Ver21AndViewPitchOnly_DONOTUSE)
 	{
 		return UE::Net::ReadQuantizedVector(ScaleFactor, Value, Ar);
 	}
@@ -250,7 +250,9 @@ bool ReadPackedVector(FVector3d& Value, FArchive& Ar)
 template<int32 ScaleFactor, int32 MaxBitsPerComponent>
 bool SerializePackedVector(FVector3f& Value, FArchive& Ar)
 {
-	if (Ar.EngineNetVer() >= HISTORY_PACKED_VECTOR_LWC_SUPPORT && Ar.EngineNetVer() != HISTORY_21_AND_VIEWPITCH_ONLY_DO_NOT_USE)
+	Ar.UsingCustomVersion(FEngineNetworkCustomVersion::Guid);
+
+	if (Ar.EngineNetVer() >= FEngineNetworkCustomVersion::PackedVectorLWCSupport && Ar.EngineNetVer() != FEngineNetworkCustomVersion::Ver21AndViewPitchOnly_DONOTUSE)
 	{
 		return UE::Net::SerializeQuantizedVector<ScaleFactor>(Value, Ar);
 	}
@@ -264,7 +266,9 @@ bool SerializePackedVector(FVector3f& Value, FArchive& Ar)
 template<int32 ScaleFactor, int32 MaxBitsPerComponent>
 bool SerializePackedVector(FVector3d& Value, FArchive& Ar)
 {
-	if (Ar.EngineNetVer() >= HISTORY_PACKED_VECTOR_LWC_SUPPORT && Ar.EngineNetVer() != HISTORY_21_AND_VIEWPITCH_ONLY_DO_NOT_USE)
+	Ar.UsingCustomVersion(FEngineNetworkCustomVersion::Guid);
+
+	if (Ar.EngineNetVer() >= FEngineNetworkCustomVersion::PackedVectorLWCSupport && Ar.EngineNetVer() != FEngineNetworkCustomVersion::Ver21AndViewPitchOnly_DONOTUSE)
 	{
 		return UE::Net::SerializeQuantizedVector<ScaleFactor>(Value, Ar);
 	}
diff --git a/Engine/Source/Runtime/Engine/Private/DataChannel.cpp b/Engine/Source/Runtime/Engine/Private/DataChannel.cpp
--- a/Engine/Source/Runtime/Engine/Private/DataChannel.cpp
+++ b/Engine/Source/Runtime/Engine/Private/DataChannel.cpp
@@ -4467,7 +4467,7 @@ UObject* UActorChannel::ReadContentBlockHeader(FInBunch& Bunch, bool& bObjectDel
 	bool bSerializeClass = true;
 	ESubObjectDeleteFlag DeleteFlag = ESubObjectDeleteFlag::Destroyed;
 
-	if (Bunch.EngineNetVer() >= HISTORY_SUBOBJECT_DESTROY_FLAG)
+	if (Bunch.EngineNetVer() >= FEngineNetworkCustomVersion::SubObjectDestroyFlag)
 	{
 		const bool bIsDestroyMessage = Bunch.ReadBit() != 0 ? true : false;
 
@@ -4576,7 +4576,7 @@ UObject* UActorChannel::ReadContentBlockHeader(FInBunch& Bunch, bool& bObjectDel
 	}
 
 	UObject* ObjOuter = Actor;
-	if (Bunch.EngineNetVer() >= HISTORY_SUBOBJECT_OUTER_CHAIN)
+	if (Bunch.EngineNetVer() >= FEngineNetworkCustomVersion::SubObjectOuterChain)
 	{
 		const bool bActorIsOuter = Bunch.ReadBit() != 0;
 
@@ -4919,7 +4919,7 @@ FNetFieldExportGroup* UActorChannel::GetNetFieldExportGroupForClassNetCache(UCla
 
 	FString NetFieldExportGroupName;
 	
-	if (Connection->EngineNetworkProtocolVersion < HISTORY_CLASSNETCACHE_FULLNAME)
+	if (Connection->GetNetworkCustomVersion(FEngineNetworkCustomVersion::Guid) < FEngineNetworkCustomVersion::ClassNetCacheFullName)
 	{
 		NetFieldExportGroupName = ObjectClass->GetName() + ClassNetCacheSuffix;
 	}
diff --git a/Engine/Source/Runtime/Engine/Private/DemoNetDriver.cpp b/Engine/Source/Runtime/Engine/Private/DemoNetDriver.cpp
--- a/Engine/Source/Runtime/Engine/Private/DemoNetDriver.cpp
+++ b/Engine/Source/Runtime/Engine/Private/DemoNetDriver.cpp
@@ -4,11 +4,6 @@
 	UDemoNetDriver.cpp: Simulated network driver for recording and playing back game sessions.
 =============================================================================*/
 
-
-// @todo: LowLevelSend now includes the packet size in bits, but this is ignored locally.
-//			Tracking of this must be added, if demos are to support PacketHandler's in the future (not presently needed).
-
-
 #include "Engine/DemoNetDriver.h"
 #include "Async/TaskGraphInterfaces.h"
 #include "Engine/DemoNetConnection.h"
@@ -3234,7 +3229,7 @@ void UDemoNetDriver::ReplayStreamingReady(const FStartStreamingResult& Result)
 			UE_LOG(LogDemo, Log, TEXT("ReplayStreamingReady: playing back replay [%s] %s, which was recorded on engine version %s with flags [%s]"),
 				*ReplayHelper.GetPlaybackGuid().ToString(EGuidFormats::Digits), *ReplayHelper.DemoURL.Map, *ReplayHelper.PlaybackDemoHeader.EngineVersion.ToString(), *HeaderFlags);
 
-			if (ReplayHelper.PlaybackDemoHeader.Version >= ENetworkVersionHistory::HISTORY_RECORDING_METADATA)
+			if (GetPlaybackReplayVersion() >= FReplayCustomVersion::RecordingMetadata)
 			{
 				UE_LOG(LogDemo, Log, TEXT("ReplayStreamingReady: replay was recorded with: MinHz: %0.2f MaxHz: %0.2f FrameMS: %0.2f CheckpointMS: %0.2f Platform: [%s] Config: [%s] Target: [%s]"),
 					ReplayHelper.PlaybackDemoHeader.MinRecordHz, ReplayHelper.PlaybackDemoHeader.MaxRecordHz,
@@ -3681,7 +3676,7 @@ bool UDemoNetDriver::FastForwardLevels(const FGotoResult& GotoResult)
 
 			TGuardValue<bool> LoadingCheckpointGuard(ReplayHelper.bIsLoadingCheckpoint, true);
 
-			uint32 PlaybackVersion = GetPlaybackDemoVersion();
+			FReplayCustomVersion::Type PlaybackReplayVersion = GetPlaybackReplayVersion();
 
 			do 
 			{
@@ -3972,7 +3967,7 @@ bool UDemoNetDriver::LoadCheckpoint(const FGotoResult& GotoResult)
 	LastProcessedPacketTime = 0.f;
 	ReplayHelper.LatestReadFrameTime = 0.f;
 
-	uint32 PlaybackVersion = GetPlaybackDemoVersion();
+	FReplayCustomVersion::Type PlaybackReplayVersion = GetPlaybackReplayVersion();
 
 	if (GotoCheckpointArchive->TotalSize() > 0)
 	{
@@ -4333,7 +4328,7 @@ bool UDemoNetDriver::LoadCheckpoint(const FGotoResult& GotoResult)
 
 			FString PathName;
 
-			if (PlaybackVersion < HISTORY_GUID_NAMETABLE)
+			if (PlaybackReplayVersion < FReplayCustomVersion::GuidNameTable)
 			{
 				*GotoCheckpointArchive << PathName;
 			}
@@ -4371,7 +4366,7 @@ bool UDemoNetDriver::LoadCheckpoint(const FGotoResult& GotoResult)
 
 			CacheObject.PathName = FName(*PathName);
 
-			if (PlaybackVersion < HISTORY_GUIDCACHE_CHECKSUMS)
+			if (PlaybackReplayVersion < FReplayCustomVersion::GuidCacheChecksums)
 			{
 				*GotoCheckpointArchive << CacheObject.NetworkChecksum;
 			}
diff --git a/Engine/Source/Runtime/Engine/Private/Engine/HitResult.cpp b/Engine/Source/Runtime/Engine/Private/Engine/HitResult.cpp
--- a/Engine/Source/Runtime/Engine/Private/Engine/HitResult.cpp
+++ b/Engine/Source/Runtime/Engine/Private/Engine/HitResult.cpp
@@ -110,7 +110,7 @@ bool FHitResult::NetSerialize(FArchive& Ar, class UPackageMap* Map, bool& bOutSu
 
 	Ar << PhysMaterial;
 
-	if (Ar.IsLoading() && Ar.EngineNetVer() < HISTORY_HITRESULT_INSTANCEHANDLE)
+	if (Ar.IsLoading() && Ar.EngineNetVer() < FEngineNetworkCustomVersion::HitResultInstanceHandle)
 	{
 		AActor* HitActor = nullptr;
 		Ar << HitActor;
diff --git a/Engine/Source/Runtime/Engine/Private/Engine/ReplicatedState.cpp b/Engine/Source/Runtime/Engine/Private/Engine/ReplicatedState.cpp
--- a/Engine/Source/Runtime/Engine/Private/Engine/ReplicatedState.cpp
+++ b/Engine/Source/Runtime/Engine/Private/Engine/ReplicatedState.cpp
@@ -64,8 +64,10 @@ bool FRepMovement::SerializeQuantizedVector(FArchive& Ar, FVector& Vector, EVect
 
 bool FRepMovement::NetSerialize(FArchive& Ar, class UPackageMap* Map, bool& bOutSuccess)
 {
+	Ar.UsingCustomVersion(FEngineNetworkCustomVersion::Guid);
+
 	// pack bitfield with flags
-	const bool bServerFrameAndHandleSupported = Ar.EngineNetVer() >= HISTORY_REPMOVE_SERVERFRAME_AND_HANDLE && Ar.EngineNetVer() != HISTORY_21_AND_VIEWPITCH_ONLY_DO_NOT_USE;
+	const bool bServerFrameAndHandleSupported = Ar.EngineNetVer() >= FEngineNetworkCustomVersion::RepMoveServerFrameAndHandle && Ar.EngineNetVer() != FEngineNetworkCustomVersion::Ver21AndViewPitchOnly_DONOTUSE;
 	uint8 Flags = (bSimulatedPhysicSleep << 0) | (bRepPhysics << 1) | ((ServerFrame > 0) << 2) | ((ServerPhysicsHandle != INDEX_NONE) << 3);
 	Ar.SerializeBits(&Flags, bServerFrameAndHandleSupported ? 4 : 2);
 	bSimulatedPhysicSleep = (Flags & (1 << 0)) ? 1 : 0;
diff --git a/Engine/Source/Runtime/Engine/Private/GameStateBase.cpp b/Engine/Source/Runtime/Engine/Private/GameStateBase.cpp
--- a/Engine/Source/Runtime/Engine/Private/GameStateBase.cpp
+++ b/Engine/Source/Runtime/Engine/Private/GameStateBase.cpp
@@ -167,7 +167,7 @@ void AGameStateBase::OnRep_ReplicatedWorldTimeSeconds()
 		const UDemoNetDriver* DemoNetDriver = World ? World->GetDemoNetDriver() : nullptr;
 
 		// Support old replays.
-		if (DemoNetDriver && DemoNetDriver->IsPlaying() && (DemoNetDriver->GetPlaybackEngineNetworkProtocolVersion() < EEngineNetworkVersionHistory::HISTORY_GAMESTATE_REPLCIATED_TIME_AS_DOUBLE))
+		if (DemoNetDriver && DemoNetDriver->IsPlaying() && (DemoNetDriver->GetPlaybackEngineNetworkProtocolVersion() < FEngineNetworkCustomVersion::GameStateReplicatedTimeAsDouble))
 		{
 			PRAGMA_DISABLE_DEPRECATION_WARNINGS
 			ReplicatedWorldTimeSecondsDouble = ReplicatedWorldTimeSeconds;
diff --git a/Engine/Source/Runtime/Engine/Private/NetConnection.cpp b/Engine/Source/Runtime/Engine/Private/NetConnection.cpp
--- a/Engine/Source/Runtime/Engine/Private/NetConnection.cpp
+++ b/Engine/Source/Runtime/Engine/Private/NetConnection.cpp
@@ -305,13 +305,11 @@ UNetConnection::UNetConnection(const FObjectInitializer& ObjectInitializer)
 ,	OutPacketId			( 0 ) // must be initialized as OutAckPacketId + 1 so loss of first packet can be detected
 ,	OutAckPacketId		( -1 )
 ,	bLastHasServerFrameTime( false )
-	, DefaultMaxChannelSize(32767)
+,	DefaultMaxChannelSize(32767)
 ,	InitOutReliable		( 0 )
 ,	InitInReliable		( 0 )
-,	EngineNetworkProtocolVersion( FNetworkVersion::GetEngineNetworkProtocolVersion() )
-,	GameNetworkProtocolVersion( FNetworkVersion::GetGameNetworkProtocolVersion() )
-	, PackageVersionUE( GPackageFileUEVersion )
-	, PackageVersionLicenseeUE( GPackageFileLicenseeUEVersion )
+,	PackageVersionUE( GPackageFileUEVersion )
+,	PackageVersionLicenseeUE( GPackageFileLicenseeUEVersion )
 ,	ResendAllDataState( EResendAllDataState::None )
 #if !UE_BUILD_SHIPPING
 ,	ReceivedRawPacketDel()
@@ -329,6 +327,12 @@ UNetConnection::UNetConnection(const FObjectInitializer& ObjectInitializer)
 ,	bFlushingPacketOrderCache(false)
 ,	ConnectionId(0)
 {
+	PRAGMA_DISABLE_DEPRECATION_WARNINGS
+	EngineNetworkProtocolVersion = FNetworkVersion::GetEngineNetworkProtocolVersion();
+	GameNetworkProtocolVersion = FNetworkVersion::GetGameNetworkProtocolVersion();
+	PRAGMA_ENABLE_DEPRECATION_WARNINGS
+
+	NetworkCustomVersions = FNetworkVersion::GetNetworkCustomVersions();
 }
 
 UNetConnection::UNetConnection(FVTableHelper& Helper)
@@ -2576,7 +2580,7 @@ bool UNetConnection::ReadPacketInfo(FBitReader& Reader, bool bHasPacketInfoPaylo
 		bLastHasServerFrameTime = bHasServerFrameTime;
 	}
 
-	if (Reader.EngineNetVer() < HISTORY_JITTER_IN_HEADER)
+	if (Reader.EngineNetVer() < FEngineNetworkCustomVersion::JitterInHeader)
 	{
 		uint8 RemoteInKBytesPerSecondByte = 0;
 		Reader << RemoteInKBytesPerSecondByte;
@@ -2762,7 +2766,7 @@ void UNetConnection::ReceivedPacket( FBitReader& Reader, bool bIsReinjectedPacke
 
 		bool bHasPacketInfoPayload = true;
 
-		if (Reader.EngineNetVer() >= HISTORY_JITTER_IN_HEADER)
+		if (Reader.EngineNetVer() >= FEngineNetworkCustomVersion::JitterInHeader)
 		{
 			bHasPacketInfoPayload = Reader.ReadBit() == 1u;
 
@@ -2981,7 +2985,7 @@ void UNetConnection::ReceivedPacket( FBitReader& Reader, bool bIsReinjectedPacke
 		while( !Reader.AtEnd() && GetConnectionState()!=USOCK_Closed )
 		{
 			// For demo backwards compatibility, old replays still have this bit
-			if (IsInternalAck() && EngineNetworkProtocolVersion < EEngineNetworkVersionHistory::HISTORY_ACKS_INCLUDED_IN_HEADER)
+			if (IsInternalAck() && GetNetworkCustomVersion(FEngineNetworkCustomVersion::Guid) < FEngineNetworkCustomVersion::AcksIncludedInHeader)
 			{
 				const bool IsAckDummy = Reader.ReadBit() == 1u;
 			}
@@ -2999,7 +3003,7 @@ void UNetConnection::ReceivedPacket( FBitReader& Reader, bool bIsReinjectedPacke
 				Bunch.bOpen					= bControl ? Reader.ReadBit() : 0;
 				Bunch.bClose				= bControl ? Reader.ReadBit() : 0;
 			
-				if (Bunch.EngineNetVer() < HISTORY_CHANNEL_CLOSE_REASON)
+				if (Bunch.EngineNetVer() < FEngineNetworkCustomVersion::ChannelCloseReason)
 				{
 					const uint8 bDormant = Bunch.bClose ? Reader.ReadBit() : 0;
 					Bunch.CloseReason = bDormant ? EChannelCloseReason::Dormancy : EChannelCloseReason::Destroyed;
@@ -3013,7 +3017,7 @@ void UNetConnection::ReceivedPacket( FBitReader& Reader, bool bIsReinjectedPacke
 				Bunch.bIsReplicationPaused  = Reader.ReadBit();
 				Bunch.bReliable				= Reader.ReadBit();
 
-				if (Bunch.EngineNetVer() < HISTORY_MAX_ACTOR_CHANNELS_CUSTOMIZATION)
+				if (Bunch.EngineNetVer() < FEngineNetworkCustomVersion::MaxActorChannelsCustomization)
 				{
 					static const int OLD_MAX_ACTOR_CHANNELS = 10240;
 					Bunch.ChIndex = Reader.ReadInt(OLD_MAX_ACTOR_CHANNELS);
@@ -3039,7 +3043,7 @@ void UNetConnection::ReceivedPacket( FBitReader& Reader, bool bIsReinjectedPacke
 
 				// if flag is set, remap channel index values, we're fast forwarding a replay checkpoint
 				// and there should be no bunches for existing channels
-				if (IsInternalAck() && bAllowExistingChannelIndex && (Bunch.EngineNetVer() >= HISTORY_REPLAY_DORMANCY))
+				if (IsInternalAck() && bAllowExistingChannelIndex && (Bunch.EngineNetVer() >= FEngineNetworkCustomVersion::ReplayDormancy))
 				{
 					if (ChannelIndexMap.Contains(Bunch.ChIndex))
 					{
@@ -3118,7 +3122,7 @@ void UNetConnection::ReceivedPacket( FBitReader& Reader, bool bIsReinjectedPacke
 				Bunch.bPartialInitial = Bunch.bPartial ? Reader.ReadBit() : 0;
 				Bunch.bPartialFinal = Bunch.bPartial ? Reader.ReadBit() : 0;
 
-				if (Bunch.EngineNetVer() < HISTORY_CHANNEL_NAMES)
+				if (Bunch.EngineNetVer() < FEngineNetworkCustomVersion::ChannelNames)
 				{
 					uint32 ChType = (Bunch.bReliable || Bunch.bOpen) ? Reader.ReadInt(CHTYPE_MAX) : CHTYPE_None;
 					switch (ChType)
@@ -3296,7 +3300,7 @@ void UNetConnection::ReceivedPacket( FBitReader& Reader, bool bIsReinjectedPacke
 				// In that case, we can generally ignore these bunches.
 				if (IsInternalAck() && bAllowExistingChannelIndex)
 				{
-					if (Bunch.EngineNetVer() < HISTORY_REPLAY_DORMANCY)
+					if (Bunch.EngineNetVer() < FEngineNetworkCustomVersion::ReplayDormancy)
 					{
 						if (Channel)
 						{
@@ -5365,8 +5369,15 @@ void UNetConnection::NotifyActorChannelCleanedUp(UActorChannel* Channel, EChanne
 
 void UNetConnection::SetNetVersionsOnArchive(FArchive& Ar) const
 {
-	Ar.SetEngineNetVer(EngineNetworkProtocolVersion);
-	Ar.SetGameNetVer(GameNetworkProtocolVersion);
+	Ar.SetEngineNetVer(GetNetworkCustomVersion(FEngineNetworkCustomVersion::Guid));
+	Ar.SetGameNetVer(GetNetworkCustomVersion(FGameNetworkCustomVersion::Guid));
+
+	const FCustomVersionArray& AllVersions = NetworkCustomVersions.GetAllVersions();
+	for (const FCustomVersion& Version : AllVersions)
+	{
+		Ar.SetCustomVersion(Version.Key, Version.Version, Version.GetFriendlyName());
+	}
+
 	Ar.SetUEVer(PackageVersionUE);
 	Ar.SetLicenseeUEVer(PackageVersionLicenseeUE);
 	// Base archives only store FEngineVersionBase, but net connections store FEngineVersion.
@@ -5374,6 +5385,21 @@ void UNetConnection::SetNetVersionsOnArchive(FArchive& Ar) const
 	Ar.SetEngineVer(EngineVersion);
 }
 
+uint32 UNetConnection::GetNetworkCustomVersion(const FGuid& VersionGuid) const
+{
+	const FCustomVersion* CustomVer = NetworkCustomVersions.GetVersion(VersionGuid);
+	return CustomVer != nullptr ? CustomVer->Version : 0;
+}
+
+void UNetConnection::SetNetworkCustomVersions(const FCustomVersionContainer& CustomVersions)
+{
+	const FCustomVersionArray& AllVersions = CustomVersions.GetAllVersions();
+	for (const FCustomVersion& Version : AllVersions)
+	{
+		NetworkCustomVersions.SetVersion(Version.Key, Version.Version, Version.GetFriendlyName());
+	}
+}
+
 /*-----------------------------------------------------------------------------
 	USimulatedClientNetConnection.
 -----------------------------------------------------------------------------*/
diff --git a/Engine/Source/Runtime/Engine/Private/PackageMapClient.cpp b/Engine/Source/Runtime/Engine/Private/PackageMapClient.cpp
--- a/Engine/Source/Runtime/Engine/Private/PackageMapClient.cpp
+++ b/Engine/Source/Runtime/Engine/Private/PackageMapClient.cpp
@@ -537,7 +537,7 @@ bool UPackageMapClient::SerializeNewActor(FArchive& Ar, class UActorChannel *Cha
 		FNetworkGUID ArchetypeNetGUID;
 		SerializeObject(Ar, UObject::StaticClass(), Archetype, &ArchetypeNetGUID);
 
-		if (Ar.IsSaving() || (Connection && (Connection->EngineNetworkProtocolVersion >= EEngineNetworkVersionHistory::HISTORY_NEW_ACTOR_OVERRIDE_LEVEL)))
+		if (Ar.IsSaving() || (Connection && (Connection->GetNetworkCustomVersion(FEngineNetworkCustomVersion::Guid) >= FEngineNetworkCustomVersion::NewActorOverrideLevel)))
 		{
 			SerializeObject(Ar, ULevel::StaticClass(), ActorLevel);
 		}
@@ -601,7 +601,7 @@ bool UPackageMapClient::SerializeNewActor(FArchive& Ar, class UActorChannel *Cha
 				Ar.SerializeBits(&bWasSerialized, 1);
 				if (bWasSerialized)
 				{
-					if (Ar.EngineNetVer() < HISTORY_OPTIONALLY_QUANTIZE_SPAWN_INFO)
+					if (Ar.IsLoading() && Ar.EngineNetVer() < FEngineNetworkCustomVersion::OptionallyQuantizeSpawnInfo)
 					{
 						bShouldQuantize = true;
 					}
@@ -3706,7 +3706,7 @@ void FNetGUIDCache::GenerateFullNetGUIDPath_r( const FNetworkGUID& NetGUID, FStr
 
 	const FNetGuidCacheObject* CacheObject = ObjectLookup.Find( NetGUID );
 
-	if ( CacheObject == NULL )
+	if ( CacheObject == nullptr )
 	{
 		// Doh, this shouldn't be possible, but if this happens, we can't continue
 		// So warn, and return
@@ -3725,7 +3725,7 @@ void FNetGUIDCache::GenerateFullNetGUIDPath_r( const FNetworkGUID& NetGUID, FStr
 	if ( CacheObject->Object.IsValid() )
 	{
 		// Sanity check that the names match if the path was stored
-		if ( CacheObject->PathName != NAME_None && CacheObject->Object->GetName() != CacheObject->PathName.ToString() )
+		if ( !CacheObject->PathName.IsNone() && CacheObject->Object->GetFName() != CacheObject->PathName )
 		{
 			UE_LOG( LogNetPackageMap, Warning, TEXT( "GenerateFullNetGUIDPath_r: Name mismatch! %s != %s" ), *CacheObject->PathName.ToString(), *CacheObject->Object->GetName() );	
 		}
@@ -3734,7 +3734,7 @@ void FNetGUIDCache::GenerateFullNetGUIDPath_r( const FNetworkGUID& NetGUID, FStr
 	}
 	else
 	{
-		if ( CacheObject->PathName == NAME_None )
+		if (CacheObject->PathName.IsNone())
 		{
 			// This can happen when a non stably named object is NULL
 			FullPath += FString::Printf( TEXT( "[%s]EMPTY" ), *NetGUID.ToString() );
@@ -3901,7 +3901,7 @@ FArchive& operator<<(FArchive& Ar, FNetFieldExport& C)
 		Ar.SerializeIntPacked(C.Handle);
 		Ar << C.CompatibleChecksum;
 
-		if (Ar.IsLoading() && Ar.EngineNetVer() < HISTORY_NETEXPORT_SERIALIZATION)
+		if (Ar.IsLoading() && Ar.EngineNetVer() < FEngineNetworkCustomVersion::NetExportSerialization)
 		{
 			FName TempName;
 			FString TempType;
@@ -3913,7 +3913,7 @@ FArchive& operator<<(FArchive& Ar, FNetFieldExport& C)
 		}
 		else
 		{
-			if (Ar.IsLoading() && Ar.EngineNetVer() < HISTORY_NETEXPORT_SERIALIZE_FIX)
+			if (Ar.IsLoading() && Ar.EngineNetVer() < FEngineNetworkCustomVersion::NetExportSerializeFix)
 			{
 				Ar << C.ExportName;
 			}
diff --git a/Engine/Source/Runtime/Engine/Private/Pawn.cpp b/Engine/Source/Runtime/Engine/Private/Pawn.cpp
--- a/Engine/Source/Runtime/Engine/Private/Pawn.cpp
+++ b/Engine/Source/Runtime/Engine/Private/Pawn.cpp
@@ -27,6 +27,7 @@
 #include "GameFramework/GameNetworkManager.h"
 #include "GameFramework/InputSettings.h"
 #include "Engine/DemoNetDriver.h"
+#include "Misc/EngineNetworkCustomVersion.h"
 
 #include UE_INLINE_GENERATED_CPP_BY_NAME(Pawn)
 
@@ -979,7 +980,7 @@ FRotator APawn::GetBaseAimRotation() const
 			const UWorld* World = GetWorld();
 			const UDemoNetDriver* DemoNetDriver = World ? World->GetDemoNetDriver() : nullptr;
 
-			if (DemoNetDriver && DemoNetDriver->IsPlaying() && (DemoNetDriver->GetPlaybackEngineNetworkProtocolVersion() < EEngineNetworkVersionHistory::HISTORY_PAWN_REMOTEVIEWPITCH))
+			if (DemoNetDriver && DemoNetDriver->IsPlaying() && (DemoNetDriver->GetPlaybackEngineNetworkProtocolVersion() < FEngineNetworkCustomVersion::PawnRemoteViewPitch))
 			{
 				POVRot.Pitch = RemoteViewPitch;
 				POVRot.Pitch = POVRot.Pitch * 360.0f / 255.0f;
diff --git a/Engine/Source/Runtime/Engine/Private/RepLayout.cpp b/Engine/Source/Runtime/Engine/Private/RepLayout.cpp
--- a/Engine/Source/Runtime/Engine/Private/RepLayout.cpp
+++ b/Engine/Source/Runtime/Engine/Private/RepLayout.cpp
@@ -961,7 +961,7 @@ static uint32 GetRepLayoutCmdCompatibleChecksum(
 	
 	// Evolve by StaticArrayIndex (to make all unrolled static array elements unique)
 	if ((ServerConnection == nullptr) ||
-		(ServerConnection->EngineNetworkProtocolVersion >= EEngineNetworkVersionHistory::HISTORY_REPCMD_CHECKSUM_REMOVE_PRINTF))
+		(ServerConnection->GetNetworkCustomVersion(FEngineNetworkCustomVersion::Guid) >= FEngineNetworkCustomVersion::RepCmdChecksumRemovePrintf))
 	{
 		CompatibleChecksum = FCrc::MemCrc32(&StaticArrayIndex, sizeof(StaticArrayIndex), CompatibleChecksum);
 	}
@@ -972,7 +972,7 @@ static uint32 GetRepLayoutCmdCompatibleChecksum(
 
 	// Evolve by enum max value bits required
 	if ((ServerConnection == nullptr) ||
-		(ServerConnection->EngineNetworkProtocolVersion >= EEngineNetworkVersionHistory::HISTORY_ENUM_SERIALIZATION_COMPAT))
+		(ServerConnection->GetNetworkCustomVersion(FEngineNetworkCustomVersion::Guid) >= FEngineNetworkCustomVersion::EnumSerializationCompat))
 	{
 		if (const FEnumProperty* EnumProp = CastField<FEnumProperty>(Property))
 		{
@@ -4564,7 +4564,7 @@ bool FRepLayout::ReceiveCustomDeltaProperty(
 	FNetDeltaSerializeInfo& Params,
 	FStructProperty* Property) const
 {
-	if (Params.Connection->EngineNetworkProtocolVersion >= EEngineNetworkVersionHistory::HISTORY_FAST_ARRAY_DELTA_STRUCT)
+	if (Params.Connection->GetNetworkCustomVersion(FEngineNetworkCustomVersion::Guid) >= FEngineNetworkCustomVersion::FastArrayDeltaStruct)
 	{
 		Params.bSupportsFastArrayDeltaStructSerialization = !!Params.Reader->ReadBit();
 	}
diff --git a/Engine/Source/Runtime/Engine/Private/ReplayHelper.cpp b/Engine/Source/Runtime/Engine/Private/ReplayHelper.cpp
--- a/Engine/Source/Runtime/Engine/Private/ReplayHelper.cpp
+++ b/Engine/Source/Runtime/Engine/Private/ReplayHelper.cpp
@@ -92,27 +92,38 @@ TSharedPtr<INetworkReplayStreamer> FReplayHelper::Init(const FURL& URL)
 
 void FReplayHelper::SetPlaybackNetworkVersions(FArchive& Ar)
 {
-	Ar.SetEngineNetVer(PlaybackDemoHeader.EngineNetworkProtocolVersion);
-	Ar.SetGameNetVer(PlaybackDemoHeader.GameNetworkProtocolVersion);
+	Ar.SetEngineNetVer(PlaybackDemoHeader.GetCustomVersion(FEngineNetworkCustomVersion::Guid));
+	Ar.SetGameNetVer(PlaybackDemoHeader.GetCustomVersion(FGameNetworkCustomVersion::Guid));
+
 	Ar.SetUEVer(PlaybackDemoHeader.PackageVersionUE);
 	Ar.SetLicenseeUEVer(PlaybackDemoHeader.PackageVersionLicenseeUE);
 	// Base archives only store FEngineVersionBase, but the header stores FEngineVersion.
 	// This will slice off the branch name and anything else stored in FEngineVersion.
 	Ar.SetEngineVer(PlaybackDemoHeader.EngineVersion);
+
+	Ar.SetCustomVersions(PlaybackDemoHeader.CustomVersions);
 }
 
 void FReplayHelper::SetPlaybackNetworkVersions(UNetConnection* Connection)
 {
 	if (Connection)
 	{
-		Connection->EngineNetworkProtocolVersion = PlaybackDemoHeader.EngineNetworkProtocolVersion;
-		Connection->GameNetworkProtocolVersion = PlaybackDemoHeader.GameNetworkProtocolVersion;
+		PRAGMA_DISABLE_DEPRECATION_WARNINGS
+		Connection->EngineNetworkProtocolVersion = PlaybackDemoHeader.GetCustomVersion(FEngineNetworkCustomVersion::Guid);
+		Connection->GameNetworkProtocolVersion = PlaybackDemoHeader.GetCustomVersion(FGameNetworkCustomVersion::Guid);
+		PRAGMA_ENABLE_DEPRECATION_WARNINGS
+		Connection->SetNetworkCustomVersions(PlaybackDemoHeader.CustomVersions);
 		Connection->SetPackageVersionUE(PlaybackDemoHeader.PackageVersionUE);
 		Connection->SetPackageVersionLicenseeUE(PlaybackDemoHeader.PackageVersionLicenseeUE);
 		Connection->SetEngineVersion(PlaybackDemoHeader.EngineVersion);
 	}
 }
 
+FReplayCustomVersion::Type FReplayHelper::GetPlaybackReplayVersion() const
+{
+	return (FReplayCustomVersion::Type)PlaybackDemoHeader.GetCustomVersion(FReplayCustomVersion::Guid);
+}
+
 void FReplayHelper::OnStartRecordingComplete(const FStartStreamingResult& Result)
 {
 	check(Result.bRecording);
@@ -202,7 +213,8 @@ void FReplayHelper::WriteNetworkDemoHeader(UNetConnection* Connection)
 
 		FNetworkReplayDelegates::OnWriteGameSpecificDemoHeader.Broadcast(DemoHeader.GameSpecificData);
 
-		if (UWorld* LocalWorld = World.Get())
+		UWorld* LocalWorld = World.Get();
+		if (LocalWorld)
 		{
 			// intentionally not checking for a demo net driver
 			if (LocalWorld->GetNetDriver() != nullptr && !LocalWorld->GetNetDriver()->IsServer())
@@ -395,7 +407,9 @@ bool FReplayHelper::ReadPlaybackDemoHeader(FString& Error)
 {
 	UGameInstance* GameInstance = World->GetGameInstance();
 
+	PRAGMA_DISABLE_DEPRECATION_WARNINGS
 	PlaybackDemoHeader = FNetworkDemoHeader();
+	PRAGMA_ENABLE_DEPRECATION_WARNINGS
 
 	FArchive* FileAr = ReplayStreamer->GetHeaderArchive();
 	if (!FileAr)
@@ -406,6 +420,10 @@ bool FReplayHelper::ReadPlaybackDemoHeader(FString& Error)
 		return false;
 	}
 
+	// sanity checking for string/array sizes when reading the header
+	FileAr->ArIsNetArchive = true;
+	FileAr->ArMaxSerializeSize = FReplayHelper::MAX_DEMO_STRING_SERIALIZATION_SIZE;
+
 	(*FileAr) << PlaybackDemoHeader;
 
 	if (FileAr->IsError())
@@ -1685,16 +1703,6 @@ bool FReplayHelper::SetExternalDataForObject(UNetConnection* Connection, UObject
 	return false;
 }
 
-void FDeltaCheckpointData::CountBytes(FArchive& Ar) const
-{
-	GRANULAR_NETWORK_MEMORY_TRACKING_INIT(Ar, "FDeltaCheckpointData::CountBytes");
-
-	GRANULAR_NETWORK_MEMORY_TRACKING_TRACK("RecordingDeletedNetStartupActors", RecordingDeletedNetStartupActors.CountBytes(Ar));
-	GRANULAR_NETWORK_MEMORY_TRACKING_TRACK("DestroyedNetStartupActors", DestroyedNetStartupActors.CountBytes(Ar));
-	GRANULAR_NETWORK_MEMORY_TRACKING_TRACK("DestroyedDynamicActors", DestroyedDynamicActors.CountBytes(Ar));
-	GRANULAR_NETWORK_MEMORY_TRACKING_TRACK("ChannelsToClose", ChannelsToClose.CountBytes(Ar));
-}
-
 void FReplayHelper::FCheckpointSaveStateContext::CountBytes(FArchive& Ar) const
 {
 	GRANULAR_NETWORK_MEMORY_TRACKING_INIT(Ar, "FCheckpointSaveStateContext::CountBytes");
diff --git a/Engine/Source/Runtime/Engine/Private/ReplayTypes.cpp b/Engine/Source/Runtime/Engine/Private/ReplayTypes.cpp
new file mode 100644
--- /dev/null
+++ b/Engine/Source/Runtime/Engine/Private/ReplayTypes.cpp
@@ -0,0 +1,229 @@
+// Copyright Epic Games, Inc. All Rights Reserved.
+
+#include "ReplayTypes.h"
+#include "Misc/NetworkVersion.h"
+#include "Net/NetworkGranularMemoryLogging.h"
+
+const FGuid FReplayCustomVersion::Guid = FGuid(0x8417998A, 0xBBC043EC, 0x81B3D119, 0x072D2722);
+FCustomVersionRegistration GRegisterReplayCustomVersion(FReplayCustomVersion::Guid, FReplayCustomVersion::LatestVersion, TEXT("Replay"));
+
+void FDeltaCheckpointData::CountBytes(FArchive& Ar) const
+{
+	GRANULAR_NETWORK_MEMORY_TRACKING_INIT(Ar, "FDeltaCheckpointData::CountBytes");
+
+	GRANULAR_NETWORK_MEMORY_TRACKING_TRACK("RecordingDeletedNetStartupActors", RecordingDeletedNetStartupActors.CountBytes(Ar));
+	GRANULAR_NETWORK_MEMORY_TRACKING_TRACK("DestroyedNetStartupActors", DestroyedNetStartupActors.CountBytes(Ar));
+	GRANULAR_NETWORK_MEMORY_TRACKING_TRACK("DestroyedDynamicActors", DestroyedDynamicActors.CountBytes(Ar));
+	GRANULAR_NETWORK_MEMORY_TRACKING_TRACK("ChannelsToClose", ChannelsToClose.CountBytes(Ar));
+}
+
+PRAGMA_DISABLE_DEPRECATION_WARNINGS
+FNetworkDemoHeader::FNetworkDemoHeader() :
+	Magic(NETWORK_DEMO_MAGIC),
+	Version((uint32)FReplayCustomVersion::CustomVersions),
+	NetworkChecksum(FNetworkVersion::GetLocalNetworkVersion()),
+	EngineNetworkProtocolVersion(FNetworkVersion::GetEngineNetworkProtocolVersion()),
+	GameNetworkProtocolVersion(FNetworkVersion::GetGameNetworkProtocolVersion()),
+	Guid(),
+	MinRecordHz(0.0f),
+	MaxRecordHz(0.0f),
+	FrameLimitInMS(0.0f),
+	CheckpointLimitInMS(0.0f),
+	BuildConfig(EBuildConfiguration::Unknown),
+	BuildTarget(EBuildTargetType::Unknown),
+	EngineVersion(FEngineVersion::Current()),
+	HeaderFlags(EReplayHeaderFlags::None),
+	PackageVersionUE(GPackageFileUEVersion),
+	PackageVersionLicenseeUE(GPackageFileLicenseeUEVersion)
+{
+	CustomVersions = FNetworkVersion::GetNetworkCustomVersions();
+
+	CustomVersions.SetVersion(FReplayCustomVersion::Guid, FReplayCustomVersion::LatestVersion, TEXT("Replay"));
+}
+PRAGMA_ENABLE_DEPRECATION_WARNINGS
+
+uint32 FNetworkDemoHeader::GetCustomVersion(const FGuid& VersionGuid) const
+{
+	const FCustomVersion* CustomVer = CustomVersions.GetVersion(VersionGuid);
+	return CustomVer != nullptr ? CustomVer->Version : 0;
+}
+
+void FNetworkDemoHeader::CountBytes(FArchive& Ar) const 
+{
+	GRANULAR_NETWORK_MEMORY_TRACKING_INIT(Ar, "FNetworkDemoHeader::CountBytes");
+
+	GRANULAR_NETWORK_MEMORY_TRACKING_TRACK("LevelNamesAndTimes",
+		LevelNamesAndTimes.CountBytes(Ar);
+		for (const FLevelNameAndTime& LevelNameAndTime : LevelNamesAndTimes)
+		{
+			LevelNameAndTime.CountBytes(Ar);
+		}
+	);
+
+	GRANULAR_NETWORK_MEMORY_TRACKING_TRACK("GameSpecificData",
+		GameSpecificData.CountBytes(Ar);
+		for (const FString& Datum : GameSpecificData)
+		{
+			Datum.CountBytes(Ar);
+		}
+	);
+}
+
+FArchive& operator<<(FArchive& Ar, FNetworkDemoHeader& Header)
+{
+	Ar << Header.Magic;
+
+	// Check magic value
+	if (Header.Magic != NETWORK_DEMO_MAGIC)
+	{
+		UE_LOG(LogDemo, Error, TEXT("Header.Magic != NETWORK_DEMO_MAGIC"));
+		Ar.SetError();
+		return Ar;
+	}
+
+	PRAGMA_DISABLE_DEPRECATION_WARNINGS
+	Ar << Header.Version;
+
+	// Check legacy version
+	if (Header.Version < (uint32)FReplayCustomVersion::MinSupportedVersion)
+	{
+		UE_LOG(LogDemo, Error, TEXT("Header.Version < MinSupportedVersion. Header.Version: %i, MinSupportedVersion: %i"), Header.Version, FReplayCustomVersion::MinSupportedVersion);
+		Ar.SetError();
+		return Ar;
+	}
+
+	FReplayCustomVersion::Type ReplayVersion = (FReplayCustomVersion::Type)Header.Version;
+
+	if (Ar.IsLoading())
+	{
+		Header.CustomVersions.Empty();
+	}
+
+	if (ReplayVersion >= FReplayCustomVersion::CustomVersions)
+	{
+		checkf(!Ar.IsSaving() || Header.CustomVersions.GetVersion(FReplayCustomVersion::Guid) != nullptr, TEXT("A valid replay custom version should exist when saving."));
+
+		Header.CustomVersions.Serialize(Ar);
+
+		if (Ar.IsLoading())
+		{
+			Ar.SetCustomVersions(Header.CustomVersions);
+
+			ReplayVersion = (FReplayCustomVersion::Type)Ar.CustomVer(FReplayCustomVersion::Guid);
+		}
+		else
+		{
+			Ar.UsingCustomVersion(FReplayCustomVersion::Guid);
+			Ar.UsingCustomVersion(FEngineNetworkCustomVersion::Guid);
+			Ar.UsingCustomVersion(FGameNetworkCustomVersion::Guid);
+		}
+	}
+
+	Ar << Header.NetworkChecksum;
+	Ar << Header.EngineNetworkProtocolVersion;
+	Ar << Header.GameNetworkProtocolVersion;
+
+	if (ReplayVersion < FReplayCustomVersion::CustomVersions)
+	{
+		if (Ar.IsLoading())
+		{
+			Ar.SetCustomVersion(FReplayCustomVersion::Guid, Header.Version, TEXT("Replay"));
+			Ar.SetCustomVersion(FEngineNetworkCustomVersion::Guid, Header.EngineNetworkProtocolVersion, TEXT("EngineNetwork"));
+			Ar.SetCustomVersion(FGameNetworkCustomVersion::Guid, Header.GameNetworkProtocolVersion, TEXT("GameNetwork"));
+
+			Header.CustomVersions.SetVersion(FReplayCustomVersion::Guid, Header.Version, TEXT("Replay"));
+			Header.CustomVersions.SetVersion(FEngineNetworkCustomVersion::Guid, Header.EngineNetworkProtocolVersion, TEXT("EngineNetwork"));
+			Header.CustomVersions.SetVersion(FGameNetworkCustomVersion::Guid, Header.GameNetworkProtocolVersion, TEXT("GameNetwork"));
+		}
+	}
+	PRAGMA_ENABLE_DEPRECATION_WARNINGS
+
+	if (Ar.IsLoading())
+	{
+		TArray<FCustomVersionDifference> VersionDiffs = FCurrentCustomVersions::Compare(Header.CustomVersions.GetAllVersions(), TEXT("Replay"));
+		for (const FCustomVersionDifference& Diff : VersionDiffs)
+		{
+			if (Diff.Type == ECustomVersionDifference::Missing)
+			{
+				UE_LOG(LogDemo, Error, TEXT("Replay was saved with a custom version that is not present. Tag %s Version %d"), *Diff.Version->Key.ToString(), Diff.Version->Version);
+				Ar.SetError();
+				return Ar;
+			}
+			else if (Diff.Type == ECustomVersionDifference::Invalid)
+			{
+				UE_LOG(LogDemo, Error, TEXT("Replay was saved with an invalid custom version. Tag %s Version %d"), *Diff.Version->Key.ToString(), Diff.Version->Version);
+				Ar.SetError();
+				return Ar;
+			}
+			else if (Diff.Type == ECustomVersionDifference::Newer)
+			{
+				FCustomVersion MaxExpectedVersion = FCurrentCustomVersions::Get(Diff.Version->Key).GetValue();
+
+				UE_LOG(LogDemo, Error, TEXT("Replay was saved with a newer custom version than the current. Tag %s Name '%s' ReplayVersion %d  MaxExpected %d"),
+					*Diff.Version->Key.ToString(), *MaxExpectedVersion.GetFriendlyName().ToString(), Diff.Version->Version, MaxExpectedVersion.Version);
+				Ar.SetError();
+				return Ar;
+			}
+		}
+	}
+
+	Ar << Header.Guid;
+	Ar << Header.EngineVersion;
+
+	if (Ar.CustomVer(FReplayCustomVersion::Guid) >= FReplayCustomVersion::SavePackageVersionUE)
+	{
+		Ar << Header.PackageVersionUE;
+		Ar << Header.PackageVersionLicenseeUE;
+	}
+	else
+	{
+		if (Ar.IsLoading())
+		{
+			// Fix up for LWC compatibility.
+			// Vectors were using operator<< to serialize in some network cases (instead of
+			// the NetSerialize function), but this operator uses EUnrealEngineObjectUE5Version
+			// and not EEngineNetworkVersionHistory for versioning. Therefore, pre-LWC replays
+			// that did not save the EUnrealEngineObjectUE5Version will try to read doubles
+			// from these vectors instead of floats.
+			//
+			// However, EEngineNetworkVersionHistory::HISTORY_SERIALIZE_DOUBLE_VECTORS_AS_DOUBLES was
+			// added around the same time as EUnrealEngineObjectUE5Version::LARGE_WORLD_COORDINATES,
+			// so we use the network version as an approximation the package version to allow
+			// most pre-LWC replays to play back correctly. The compromise is that any replays recorded
+			// between when HISTORY_SERIALIZE_DOUBLE_VECTORS_AS_DOUBLES and LARGE_WORLD_COORDINATES
+			// were added won't play back correctly since they didn't store accurate version information.
+			if (Ar.EngineNetVer() < FEngineNetworkCustomVersion::SerializeDoubleVectorsAsDoubles)
+			{
+				// If the replay was recorded before vectors were serialized with LWC, and before replays
+				// saved the UE package version, set the package version to one before LARGE_WORLD_COORDINATES
+				// so that vectors will be read properly by operator<<.
+				Header.PackageVersionUE = FPackageFileVersion(VER_LATEST_ENGINE_UE4, EUnrealEngineObjectUE5Version::OPTIONAL_RESOURCES);
+			}
+		}
+	}
+
+	Ar << Header.LevelNamesAndTimes;
+	Ar << Header.HeaderFlags;
+	Ar << Header.GameSpecificData;
+
+	if (Ar.CustomVer(FReplayCustomVersion::Guid) >= FReplayCustomVersion::RecordingMetadata)
+	{
+		Ar << Header.MinRecordHz;
+		Ar << Header.MaxRecordHz;
+
+		Ar << Header.FrameLimitInMS;
+		Ar << Header.CheckpointLimitInMS;
+
+		Ar << Header.Platform;
+		Ar << Header.BuildConfig;
+		Ar << Header.BuildTarget;
+	}
+
+	return Ar;
+}
+
+uint32 FOverridableReplayVersionData::GetCustomVersion(const FGuid& VersionGuid) const
+{
+	const FCustomVersion* CustomVer = CustomVersions.GetVersion(VersionGuid);
+	return CustomVer != nullptr ? CustomVer->Version : 0;
+}
diff --git a/Engine/Source/Runtime/Engine/Public/ReplayHelper.h b/Engine/Source/Runtime/Engine/Public/ReplayHelper.h
--- a/Engine/Source/Runtime/Engine/Public/ReplayHelper.h
+++ b/Engine/Source/Runtime/Engine/Public/ReplayHelper.h
@@ -56,6 +56,8 @@ private:
 	void SetPlaybackNetworkVersions(FArchive& Ar);
 	void SetPlaybackNetworkVersions(UNetConnection* Connection);
 
+	FReplayCustomVersion::Type GetPlaybackReplayVersion() const;
+
 	void StartRecording(UNetConnection* Connection);
 	void StopReplay();
 
diff --git a/Engine/Source/Runtime/Engine/Public/ReplayTypes.h b/Engine/Source/Runtime/Engine/Public/ReplayTypes.h
--- a/Engine/Source/Runtime/Engine/Public/ReplayTypes.h
+++ b/Engine/Source/Runtime/Engine/Public/ReplayTypes.h
@@ -12,6 +12,7 @@
 #include "Net/ReplayResult.h"
 #include "IPAddress.h"
 #include "Serialization/BitReader.h"
+#include "Serialization/CustomVersion.h"
 #include "Traits/IsCharEncodingCompatibleWith.h"
 #include "ReplayTypes.generated.h"
 
@@ -90,7 +91,7 @@ struct FLevelNameAndTime
 	}
 };
 
-enum ENetworkVersionHistory
+enum UE_DEPRECATED(5.2, "Using custom versions instead going forward, see FReplayCustomVersion") ENetworkVersionHistory
 {
 	HISTORY_REPLAY_INITIAL = 1,
 	HISTORY_SAVE_ABS_TIME_MS = 2,				// We now save the abs demo time in ms for each frame (solves accumulation errors)
@@ -110,26 +111,85 @@ enum ENetworkVersionHistory
 	HISTORY_GUIDCACHE_CHECKSUMS = 16,			// Removing guid export checksums from saved data, they are ignored during playback
 	HISTORY_SAVE_PACKAGE_VERSION_UE = 17,		// Save engine and licensee package version as well, in case serialization functions need them for compatibility
 	HISTORY_RECORDING_METADATA = 18,			// Adding additional record-time information to the header
+	HISTORY_USE_CUSTOM_VERSION = 19,			// Serializing replay and network versions as custom verions going forward
 
 	// -----<new versions can be added before this line>-------------------------------------------------
 	HISTORY_PLUS_ONE,
 	HISTORY_LATEST = HISTORY_PLUS_ONE - 1
 };
 
-inline static const uint32 NETWORK_DEMO_MAGIC = 0x2CF5A13D;
-inline static const uint32 NETWORK_DEMO_VERSION = HISTORY_LATEST;
-inline static const uint32 MIN_NETWORK_DEMO_VERSION = HISTORY_CHARACTER_MOVEMENT;
+struct ENGINE_API FReplayCustomVersion
+{
+	enum Type
+	{
+		// Before any version changes were made
+		BeforeCustomVersionWasAdded = 0,
+
+		// Original replay versions from ENetworkVersionHistory
+		ReplayInitial = 1,
+		SaveAbsTimeMs = 2,					// We now save the abs demo time in ms for each frame (solves accumulation errors)
+		IncreaseBuffer = 3,					// Increased buffer size of packets, which invalidates old replays
+		SaveEngineVersion = 4,				// Now saving engine net version + InternalProtocolVersion
+		ExtraVersion = 5,					// We now save engine/game protocol version, checksum, and changelist
+		MultipleLevels = 6,					// Replays support seamless travel between levels
+		MultipleLvelsTimeChanges,			// Save out the time that level changes happen
+		DeletedStartupActors = 8,			// Save DeletedNetStartupActors inside checkpoints
+		HeaderFlags = 9,					// Save out enum flags with demo header
+		LevelStreamingFixes = 10,			// Optional level streaming fixes.
+		SaveFullEngineVersion = 11,			// Now saving the entire FEngineVersion including branch name
+		HeaderGuid = 12,					// Save guid to demo header
+		CharacterMovement = 13,				// Change to using replicated movement and not interpolation
+		CharacterMovementNoInterp = 14,		// No longer recording interpolated movement samples
+		GuidNameTable = 15,					// Added a string table for exported guids
+		GuidCacheChecksums = 16,			// Removing guid export checksums from saved data, they are ignored during playback
+		SavePackageVersionUE = 17,			// Save engine and licensee package version as well, in case serialization functions need them for compatibility
+		RecordingMetadata = 18,				// Adding additional record-time information to the header
+		CustomVersions = 19,				// Serializing replay and network versions as custom verions going forward
+
+		// -----<new versions can be added above this line>-------------------------------------------------
+		VersionPlusOne,
+		LatestVersion = VersionPlusOne - 1,
+
+		MinSupportedVersion = CharacterMovement	// Minimum supported playback version
+	};
+
+	// The GUID for this custom version number
+	const static FGuid Guid;
+
+	FReplayCustomVersion() = delete;
+};
 
+inline static const uint32 NETWORK_DEMO_MAGIC = 0x2CF5A13D;
+PRAGMA_DISABLE_DEPRECATION_WARNINGS
+UE_DEPRECATED(5.2, "Now using custom versions, see FReplayCustomVersion::Latest")
+inline static const uint32 NETWORK_DEMO_VERSION = ENetworkVersionHistory::HISTORY_LATEST;
+UE_DEPRECATED(5.2, "Using custom versions instead going forward.")
+inline static const uint32 MIN_NETWORK_DEMO_VERSION = ENetworkVersionHistory::HISTORY_CHARACTER_MOVEMENT;
+PRAGMA_ENABLE_DEPRECATION_WARNINGS
+
+UE_DEPRECATED(5.2, "No longer used.")
 inline static const uint32 NETWORK_DEMO_METADATA_MAGIC = 0x3D06B24E;
+UE_DEPRECATED(5.2, "No longer used.")
 inline static const uint32 NETWORK_DEMO_METADATA_VERSION = 0;
 
-struct FNetworkDemoHeader
+struct ENGINE_API FNetworkDemoHeader
 {
 	uint32	Magic;									// Magic to ensure we're opening the right file.
+	
+	UE_DEPRECATED(5.2, "No longer used in favor of custom versions, kept for backwards compatibility.")
 	uint32	Version;								// Version number to detect version mismatches.
+
+	UE_DEPRECATED(5.2, "No longer used.")
 	uint32	NetworkChecksum;						// Network checksum
+
+	UE_DEPRECATED(5.2, "No longer used in favor of custom versions, kept for backwards compatibility.")
 	uint32	EngineNetworkProtocolVersion;			// Version of the engine internal network format
+
+	UE_DEPRECATED(5.2, "No longer used in favor of custom versions, kept for backwards compatibility.")
 	uint32	GameNetworkProtocolVersion;				// Version of the game internal network format
+
+	FCustomVersionContainer CustomVersions;
+
 	FGuid	Guid;									// Unique identifier
 
 	float MinRecordHz;
@@ -148,120 +208,20 @@ struct FNetworkDemoHeader
 	FPackageFileVersion PackageVersionUE;			// Engine package version on which the replay was recorded
 	int32 PackageVersionLicenseeUE;					// Licensee package version on which the replay was recorded
 
-	FNetworkDemoHeader() :
-		Magic(NETWORK_DEMO_MAGIC),
-		Version(NETWORK_DEMO_VERSION),
-		NetworkChecksum(FNetworkVersion::GetLocalNetworkVersion()),
-		EngineNetworkProtocolVersion(FNetworkVersion::GetEngineNetworkProtocolVersion()),
-		GameNetworkProtocolVersion(FNetworkVersion::GetGameNetworkProtocolVersion()),
-		Guid(),
-		MinRecordHz(0.0f),
-		MaxRecordHz(0.0f),
-		FrameLimitInMS(0.0f),
-		CheckpointLimitInMS(0.0f),
-		BuildConfig(EBuildConfiguration::Unknown),
-		BuildTarget(EBuildTargetType::Unknown),
-		EngineVersion(FEngineVersion::Current()),
-		HeaderFlags(EReplayHeaderFlags::None),
-		PackageVersionUE(GPackageFileUEVersion),
-		PackageVersionLicenseeUE(GPackageFileLicenseeUEVersion)
-	{
-	}
-
-	friend FArchive& operator << (FArchive& Ar, FNetworkDemoHeader& Header)
-	{
-		Ar << Header.Magic;
-
-		// Check magic value
-		if (Header.Magic != NETWORK_DEMO_MAGIC)
-		{
-			UE_LOG(LogDemo, Error, TEXT("Header.Magic != NETWORK_DEMO_MAGIC"));
-			Ar.SetError();
-			return Ar;
-		}
-
-		Ar << Header.Version;
-
-		// Check version
-		if (Header.Version < MIN_NETWORK_DEMO_VERSION)
-		{
-			UE_LOG(LogDemo, Error, TEXT("Header.Version < MIN_NETWORK_DEMO_VERSION. Header.Version: %i, MIN_NETWORK_DEMO_VERSION: %i"), Header.Version, MIN_NETWORK_DEMO_VERSION);
-			Ar.SetError();
-			return Ar;
-		}
-
-		Ar << Header.NetworkChecksum;
-		Ar << Header.EngineNetworkProtocolVersion;
-		Ar << Header.GameNetworkProtocolVersion;
-		Ar << Header.Guid;
-		Ar << Header.EngineVersion;
-
-		if (Header.Version >= HISTORY_SAVE_PACKAGE_VERSION_UE)
-		{
-			Ar << Header.PackageVersionUE;
-			Ar << Header.PackageVersionLicenseeUE;
-		}
-		else
-		{
-			if (Ar.IsLoading())
-			{
-				// Fix up for LWC compatibility.
-				// Vectors were using operator<< to serialize in some network cases (instead of
-				// the NetSerialize function), but this operator uses EUnrealEngineObjectUE5Version
-				// and not EEngineNetworkVersionHistory for versioning. Therefore, pre-LWC replays
-				// that did not save the EUnrealEngineObjectUE5Version will try to read doubles
-				// from these vectors instead of floats.
-				//
-				// However, EEngineNetworkVersionHistory::HISTORY_SERIALIZE_DOUBLE_VECTORS_AS_DOUBLES was
-				// added around the same time as EUnrealEngineObjectUE5Version::LARGE_WORLD_COORDINATES,
-				// so we use the network version as an approximation the package version to allow
-				// most pre-LWC replays to play back correctly. The compromise is that any replays recorded
-				// between when HISTORY_SERIALIZE_DOUBLE_VECTORS_AS_DOUBLES and LARGE_WORLD_COORDINATES
-				// were added won't play back correctly since they didn't store accurate version information.
-				if (Header.EngineNetworkProtocolVersion < HISTORY_SERIALIZE_DOUBLE_VECTORS_AS_DOUBLES)
-				{
-					// If the replay was recorded before vectors were serialized with LWC, and before replays
-					// saved the UE package version, set the package version to one before LARGE_WORLD_COORDINATES
-					// so that vectors will be read properly by operator<<.
-					Header.PackageVersionUE = FPackageFileVersion(VER_LATEST_ENGINE_UE4, EUnrealEngineObjectUE5Version::OPTIONAL_RESOURCES);
-				}
-			}
-		}
-
-		Ar << Header.LevelNamesAndTimes;
-		Ar << Header.HeaderFlags;
-		Ar << Header.GameSpecificData;
+	FNetworkDemoHeader();
 
-		if (Header.Version >= HISTORY_RECORDING_METADATA)
-		{
-			Ar << Header.MinRecordHz;
-			Ar << Header.MaxRecordHz;
+	PRAGMA_DISABLE_DEPRECATION_WARNINGS
+	FNetworkDemoHeader(const FNetworkDemoHeader&) = default;
+	FNetworkDemoHeader& operator=(const FNetworkDemoHeader&) = default;
+	FNetworkDemoHeader(FNetworkDemoHeader&&) = default;
+	FNetworkDemoHeader& operator=(FNetworkDemoHeader&&) = default;
+	PRAGMA_ENABLE_DEPRECATION_WARNINGS
 
-			Ar << Header.FrameLimitInMS;
-			Ar << Header.CheckpointLimitInMS;
-
-			Ar << Header.Platform;
-			Ar << Header.BuildConfig;
-			Ar << Header.BuildTarget;
-		}
-
-		return Ar;
-	}
+	ENGINE_API friend FArchive& operator << (FArchive& Ar, FNetworkDemoHeader& Header);
 
-	void CountBytes(FArchive& Ar) const
-	{
-		LevelNamesAndTimes.CountBytes(Ar);
-		for (const FLevelNameAndTime& LevelNameAndTime : LevelNamesAndTimes)
-		{
-			LevelNameAndTime.CountBytes(Ar);
-		}
+	void CountBytes(FArchive& Ar) const;
 
-		GameSpecificData.CountBytes(Ar);
-		for (const FString& Datum : GameSpecificData)
-		{
-			Datum.CountBytes(Ar);
-		}
-	}
+	uint32 GetCustomVersion(const FGuid& VersionGuid) const;
 };
 
 // The type we use to store offsets in the archive
@@ -542,23 +502,29 @@ public:
 // Using an indirect array here since FReplayExternalData stores an FBitReader, and it's not safe to store an FArchive directly in a TArray.
 typedef TIndirectArray<FReplayExternalData> FReplayExternalDataArray;
 
-
 // Can be used to override Version Data in a Replay's Header either Right Before Writing a Replay Header or Right After Reading a Replay Header.
-struct FOverridableReplayVersionData
+struct ENGINE_API FOverridableReplayVersionData
 {
 public:
+	UE_DEPRECATED(5.2, "No longer used in favor of CustomVersions")
 	uint32 Version;                       // Version number to detect version mismatches.
+	UE_DEPRECATED(5.2, "No longer used in favor of CustomVersions")
 	uint32 EngineNetworkProtocolVersion;  // Version of the engine internal network format
+	UE_DEPRECATED(5.2, "No longer used in favor of CustomVersions")
 	uint32 GameNetworkProtocolVersion;    // Version of the game internal network format
+
+	FCustomVersionContainer CustomVersions;
 	FEngineVersion EngineVersion;         // Full engine version on which the replay was recorded
 	FPackageFileVersion PackageVersionUE; // Engine package version on which the replay was recorded
 	int32 PackageVersionLicenseeUE;       // Licensee package version on which the replay was recorded
 
+	PRAGMA_DISABLE_DEPRECATION_WARNINGS
 	// Init with Demo Header Version Data
 	FOverridableReplayVersionData(const FNetworkDemoHeader& DemoHeader)
 		: Version                     (DemoHeader.Version)
 		, EngineNetworkProtocolVersion(DemoHeader.EngineNetworkProtocolVersion)
 		, GameNetworkProtocolVersion  (DemoHeader.GameNetworkProtocolVersion)
+		, CustomVersions              (DemoHeader.CustomVersions)
 		, EngineVersion               (DemoHeader.EngineVersion)
 		, PackageVersionUE            (DemoHeader.PackageVersionUE)
 		, PackageVersionLicenseeUE    (DemoHeader.PackageVersionLicenseeUE)
@@ -571,8 +537,12 @@ public:
 		DemoHeader.Version                      = Version;
 		DemoHeader.EngineNetworkProtocolVersion = EngineNetworkProtocolVersion;
 		DemoHeader.GameNetworkProtocolVersion   = GameNetworkProtocolVersion;
+		DemoHeader.CustomVersions               = CustomVersions;
 		DemoHeader.EngineVersion                = EngineVersion;
 		DemoHeader.PackageVersionUE             = PackageVersionUE;
 		DemoHeader.PackageVersionLicenseeUE     = PackageVersionLicenseeUE;
 	}
+	PRAGMA_ENABLE_DEPRECATION_WARNINGS
+
+	uint32 GetCustomVersion(const FGuid& VersionGuid) const;
 };
\ No newline at end of file
diff --git a/Engine/Source/Runtime/NetworkReplayStreaming/LocalFileNetworkReplayStreaming/Private/LocalFileNetworkReplayStreaming.cpp b/Engine/Source/Runtime/NetworkReplayStreaming/LocalFileNetworkReplayStreaming/Private/LocalFileNetworkReplayStreaming.cpp
--- a/Engine/Source/Runtime/NetworkReplayStreaming/LocalFileNetworkReplayStreaming/Private/LocalFileNetworkReplayStreaming.cpp
+++ b/Engine/Source/Runtime/NetworkReplayStreaming/LocalFileNetworkReplayStreaming/Private/LocalFileNetworkReplayStreaming.cpp
@@ -47,23 +47,8 @@ DECLARE_CYCLE_STAT(TEXT("Local replay flush stream"), STAT_LocalReplay_FlushStre
 DECLARE_CYCLE_STAT(TEXT("Local replay flush header"), STAT_LocalReplay_FlushHeader, STATGROUP_LocalReplay);
 DECLARE_CYCLE_STAT(TEXT("Local replay flush event"), STAT_LocalReplay_FlushEvent, STATGROUP_LocalReplay);
 
-namespace LocalFileReplay
+namespace UE::Net::LocalFileReplay
 {
-	enum ELocalFileVersionHistory : uint32
-	{
-		HISTORY_INITIAL							= 0,
-		HISTORY_FIXEDSIZE_FRIENDLY_NAME			= 1,
-		HISTORY_COMPRESSION						= 2,
-		HISTORY_RECORDED_TIMESTAMP				= 3,
-		HISTORY_STREAM_CHUNK_TIMES				= 4,
-		HISTORY_FRIENDLY_NAME_ENCODING			= 5,
-		HISTORY_ENCRYPTION						= 6,
-
-		// -----<new versions can be added before this line>-------------------------------------------------
-		HISTORY_PLUS_ONE,
-		HISTORY_LATEST 							= HISTORY_PLUS_ONE - 1
-	};
-
 	TAutoConsoleVariable<int32> CVarMaxCacheSize(TEXT("localReplay.MaxCacheSize"), 1024 * 1024 * 10, TEXT(""));
 	TAutoConsoleVariable<int32> CVarMaxBufferedStreamChunks(TEXT("localReplay.MaxBufferedStreamChunks"), 10, TEXT(""));
 	TAutoConsoleVariable<int32> CVarAllowLiveStreamDelete(TEXT("localReplay.AllowLiveStreamDelete"), 1, TEXT(""));
@@ -79,11 +64,30 @@ namespace LocalFileReplay
 
 const uint32 FLocalFileNetworkReplayStreamer::FileMagic = 0x1CA2E27F;
 const uint32 FLocalFileNetworkReplayStreamer::MaxFriendlyNameLen = 256;
-const uint32 FLocalFileNetworkReplayStreamer::LatestVersion = LocalFileReplay::HISTORY_LATEST;
+
+PRAGMA_DISABLE_DEPRECATION_WARNINGS
+const uint32 FLocalFileNetworkReplayStreamer::LatestVersion = FLocalFileReplayCustomVersion::LatestVersion;
+PRAGMA_ENABLE_DEPRECATION_WARNINGS
+
+const FGuid FLocalFileReplayCustomVersion::Guid =  FGuid(0x95A4f03E, 0x7E0B49E4, 0xBA43D356, 0x94FF87D9);
+FCustomVersionRegistration GRegisterLocalFileReplayCustomVersion(FLocalFileReplayCustomVersion::Guid, FLocalFileReplayCustomVersion::LatestVersion, TEXT("LocalFileReplay"));
 
 FLocalFileNetworkReplayStreamer::FLocalFileSerializationInfo::FLocalFileSerializationInfo() 
-	: FileVersion(LocalFileReplay::HISTORY_LATEST)
+	: FileVersion(FLocalFileReplayCustomVersion::CustomVersions)
 {
+	FileCustomVersions.SetVersion(FLocalFileReplayCustomVersion::Guid, FLocalFileReplayCustomVersion::LatestVersion, TEXT("LocalFileReplay"));
+}
+
+FLocalFileReplayCustomVersion::Type FLocalFileNetworkReplayStreamer::FLocalFileSerializationInfo::GetLocalFileReplayVersion() const
+{
+	if (const FCustomVersion* CustomVer = FileCustomVersions.GetVersion(FLocalFileReplayCustomVersion::Guid))
+	{
+		return (FLocalFileReplayCustomVersion::Type)CustomVer->Version;
+	}
+
+	PRAGMA_DISABLE_DEPRECATION_WARNINGS
+	return (FLocalFileReplayCustomVersion::Type)FileVersion;
+	PRAGMA_ENABLE_DEPRECATION_WARNINGS
 }
 
 void FLocalFileStreamFArchive::Serialize(void* V, int64 Length) 
@@ -204,12 +208,55 @@ bool FLocalFileNetworkReplayStreamer::ReadReplayInfo(FArchive& Archive, FLocalFi
 		uint32 MagicNumber;
 		Archive << MagicNumber;
 
-		uint32 FileVersion;
-		Archive << FileVersion;
+		uint32 DoNotUse_FileVersion;
+		Archive << DoNotUse_FileVersion;
 
 		if (MagicNumber == FLocalFileNetworkReplayStreamer::FileMagic)
 		{
-			SerializationInfo.FileVersion = FileVersion;
+			PRAGMA_DISABLE_DEPRECATION_WARNINGS
+			FCustomVersionContainer FileCustomVersions;
+
+			if (DoNotUse_FileVersion >= FLocalFileReplayCustomVersion::CustomVersions)
+			{
+				FileCustomVersions.Serialize(Archive);
+
+				TArray<FCustomVersionDifference> VersionDiffs = FCurrentCustomVersions::Compare(FileCustomVersions.GetAllVersions(), TEXT("LocalFileReplay"));
+				for (const FCustomVersionDifference& Diff : VersionDiffs)
+				{
+					if (Diff.Type == ECustomVersionDifference::Missing)
+					{
+						UE_LOG(LogLocalFileReplay, Error, TEXT("Replay was saved with a custom version that is not present. Tag %s Version %d"), *Diff.Version->Key.ToString(), Diff.Version->Version);
+						Archive.SetError();
+						return false;
+					}
+					else if (Diff.Type == ECustomVersionDifference::Invalid)
+					{
+						UE_LOG(LogLocalFileReplay, Error, TEXT("Replay was saved with an invalid custom version. Tag %s Version %d"), *Diff.Version->Key.ToString(), Diff.Version->Version);
+						Archive.SetError();
+						return false;
+					}
+					else if (Diff.Type == ECustomVersionDifference::Newer)
+					{
+						const FCustomVersion MaxExpectedVersion = FCurrentCustomVersions::Get(Diff.Version->Key).GetValue();
+
+						UE_LOG(LogLocalFileReplay, Error, TEXT("Replay was saved with a newer custom version than the current. Tag %s Name '%s' ReplayVersion %d  MaxExpected %d"),
+							*Diff.Version->Key.ToString(), *MaxExpectedVersion.GetFriendlyName().ToString(), Diff.Version->Version, MaxExpectedVersion.Version);
+						Archive.SetError();
+						return false;
+					}
+				}
+
+				Archive.SetCustomVersions(FileCustomVersions);
+			}
+			else
+			{
+				Archive.SetCustomVersion(FLocalFileReplayCustomVersion::Guid, DoNotUse_FileVersion, TEXT("LocalFileReplay"));
+				FileCustomVersions.SetVersion(FLocalFileReplayCustomVersion::Guid, DoNotUse_FileVersion, TEXT("LocalFileReplay"));
+			}
+
+			SerializationInfo.FileVersion = DoNotUse_FileVersion;
+			SerializationInfo.FileCustomVersions = FileCustomVersions;
+			PRAGMA_ENABLE_DEPRECATION_WARNINGS
 
 			// read summary info
 			Archive << Info.LengthInMS;
@@ -221,7 +268,7 @@ bool FLocalFileNetworkReplayStreamer::ReadReplayInfo(FArchive& Archive, FLocalFi
 
 			SerializationInfo.FileFriendlyName = FriendlyName;
 
-			if (FileVersion >= LocalFileReplay::HISTORY_FIXEDSIZE_FRIENDLY_NAME)
+			if (SerializationInfo.GetLocalFileReplayVersion() >= FLocalFileReplayCustomVersion::FixedSizeFriendlyName)
 			{
 				// trim whitespace since this may have been padded
 				Info.FriendlyName = FriendlyName.TrimEnd();
@@ -238,12 +285,12 @@ bool FLocalFileNetworkReplayStreamer::ReadReplayInfo(FArchive& Archive, FLocalFi
 
 			Info.bIsLive = (IsLive != 0);
 
-			if (FileVersion >= LocalFileReplay::HISTORY_RECORDED_TIMESTAMP)
+			if (SerializationInfo.GetLocalFileReplayVersion() >= FLocalFileReplayCustomVersion::RecordingTimestamp)
 			{
 				Archive << Info.Timestamp;
 			}
 
-			if (FileVersion >= LocalFileReplay::HISTORY_COMPRESSION)
+			if (SerializationInfo.GetLocalFileReplayVersion() >= FLocalFileReplayCustomVersion::CompressionSupport)
 			{
 				uint32 Compressed;
 				Archive << Compressed;
@@ -251,7 +298,7 @@ bool FLocalFileNetworkReplayStreamer::ReadReplayInfo(FArchive& Archive, FLocalFi
 				Info.bCompressed = (Compressed != 0);
 			}
 
-			if (FileVersion >= LocalFileReplay::HISTORY_ENCRYPTION)
+			if (SerializationInfo.GetLocalFileReplayVersion() >= FLocalFileReplayCustomVersion::EncryptionSupport)
 			{
 				uint32 Encrypted;
 				Archive << Encrypted;
@@ -342,7 +389,7 @@ bool FLocalFileNetworkReplayStreamer::ReadReplayInfo(FArchive& Archive, FLocalFi
 					DataChunk.ChunkIndex = Idx;
 					DataChunk.StreamOffset = Info.TotalDataSizeInBytes;
 
-					if (FileVersion >= LocalFileReplay::HISTORY_STREAM_CHUNK_TIMES)
+					if (SerializationInfo.GetLocalFileReplayVersion() >= FLocalFileReplayCustomVersion::StreamChunkTimes)
 					{
 						Archive << DataChunk.Time1;
 						Archive << DataChunk.Time2;
@@ -353,7 +400,7 @@ bool FLocalFileNetworkReplayStreamer::ReadReplayInfo(FArchive& Archive, FLocalFi
 						DataChunk.SizeInBytes = Chunk.SizeInBytes;
 					}
 
-					if (FileVersion < LocalFileReplay::HISTORY_ENCRYPTION)
+					if (SerializationInfo.GetLocalFileReplayVersion() < FLocalFileReplayCustomVersion::EncryptionSupport)
 					{
 						DataChunk.ReplayDataOffset = Archive.Tell();
 	
@@ -432,7 +479,7 @@ bool FLocalFileNetworkReplayStreamer::ReadReplayInfo(FArchive& Archive, FLocalFi
 			}
 		}
 
-		if (FileVersion < LocalFileReplay::HISTORY_STREAM_CHUNK_TIMES)
+		if (SerializationInfo.GetLocalFileReplayVersion() < FLocalFileReplayCustomVersion::StreamChunkTimes)
 		{
 			for(int i=0; i < Info.DataChunks.Num(); ++i)
 			{
@@ -528,7 +575,7 @@ bool FLocalFileNetworkReplayStreamer::AllowEncryptedWrite() const
 	bool bAllowWrite = SupportsEncryption();
 
 #if !UE_BUILD_SHIPPING
-	bAllowWrite = bAllowWrite && (LocalFileReplay::CVarAllowEncryptedRecording.GetValueOnAnyThread() != 0);
+	bAllowWrite = bAllowWrite && (UE::Net::LocalFileReplay::CVarAllowEncryptedRecording.GetValueOnAnyThread() != 0);
 
 	UE_LOG(LogLocalFileReplay, VeryVerbose, TEXT("FLocalFileNetworkReplayStreamer::AllowEncryptedWrite: %s"), *LexToString(bAllowWrite));
 #endif
@@ -538,7 +585,7 @@ bool FLocalFileNetworkReplayStreamer::AllowEncryptedWrite() const
 
 bool FLocalFileNetworkReplayStreamer::WriteReplayInfo(FArchive& Archive, const FLocalFileReplayInfo& InReplayInfo, FLocalFileSerializationInfo& SerializationInfo)
 {
-	if (SerializationInfo.FileVersion < LocalFileReplay::HISTORY_FIXEDSIZE_FRIENDLY_NAME)
+	if (SerializationInfo.GetLocalFileReplayVersion() < FLocalFileReplayCustomVersion::FixedSizeFriendlyName)
 	{
 		UE_LOG(LogLocalFileReplay, Warning, TEXT("FLocalFileNetworkRepalyStreamer::WriteReplayInfo: Unable to safely rewrite old replay info"));
 		return false;
@@ -549,8 +596,17 @@ bool FLocalFileNetworkReplayStreamer::WriteReplayInfo(FArchive& Archive, const F
 	uint32 MagicNumber = FLocalFileNetworkReplayStreamer::FileMagic;
 	Archive << MagicNumber;
 
+	PRAGMA_DISABLE_DEPRECATION_WARNINGS
 	Archive << SerializationInfo.FileVersion;
 
+	if (SerializationInfo.GetLocalFileReplayVersion() >= FLocalFileReplayCustomVersion::CustomVersions)
+	{
+		SerializationInfo.FileCustomVersions.Serialize(Archive);
+
+		Archive.SetCustomVersions(SerializationInfo.FileCustomVersions);
+	}
+	PRAGMA_ENABLE_DEPRECATION_WARNINGS
+
 	Archive << const_cast<int32&>(InReplayInfo.LengthInMS);
 	Archive << const_cast<uint32&>(InReplayInfo.NetworkVersion);
 	Archive << const_cast<uint32&>(InReplayInfo.Changelist);
@@ -558,7 +614,7 @@ bool FLocalFileNetworkReplayStreamer::WriteReplayInfo(FArchive& Archive, const F
 	FString FixedSizeName;
 	FixupFriendlyNameLength(InReplayInfo.FriendlyName, FixedSizeName);
 
-	if (SerializationInfo.FileVersion < LocalFileReplay::HISTORY_FRIENDLY_NAME_ENCODING)
+	if (SerializationInfo.GetLocalFileReplayVersion() < FLocalFileReplayCustomVersion::FriendlyNameCharEncoding)
 	{
 		// if the new name contains non-ANSI characters and the old does not, serializing would corrupt the file
 		if (!FCString::IsPureAnsi(*FixedSizeName) && FCString::IsPureAnsi(*SerializationInfo.FileFriendlyName))
@@ -595,18 +651,18 @@ bool FLocalFileNetworkReplayStreamer::WriteReplayInfo(FArchive& Archive, const F
 
 	// It's possible we're updating an older replay (e.g., for a rename)
 	// Therefore, we can't write out any data that the replay wouldn't have had.
-	if (SerializationInfo.FileVersion >= LocalFileReplay::HISTORY_RECORDED_TIMESTAMP)
+	if (SerializationInfo.GetLocalFileReplayVersion() >= FLocalFileReplayCustomVersion::RecordingTimestamp)
 	{
 		Archive << const_cast<FDateTime&>(InReplayInfo.Timestamp);
 	}
 
-	if (SerializationInfo.FileVersion >= LocalFileReplay::HISTORY_COMPRESSION)
+	if (SerializationInfo.GetLocalFileReplayVersion() >= FLocalFileReplayCustomVersion::CompressionSupport)
 	{
 		uint32 Compressed = SupportsCompression() ? 1 : 0;
 		Archive << Compressed;
 	}
 
-	if (SerializationInfo.FileVersion >= LocalFileReplay::HISTORY_ENCRYPTION)
+	if (SerializationInfo.GetLocalFileReplayVersion() >= FLocalFileReplayCustomVersion::EncryptionSupport)
 	{
 		uint32 Encrypted = AllowEncryptedWrite() ? 1 : 0;
 		Archive << Encrypted;
@@ -996,7 +1052,7 @@ void FLocalFileNetworkReplayStreamer::DeleteFinishedStream_Internal(const FStrin
 
 			const bool bIsLive = IsNamedStreamLive(StreamName);
 
-			if (LocalFileReplay::CVarAllowLiveStreamDelete.GetValueOnAnyThread() || !bIsLive)
+			if (UE::Net::LocalFileReplay::CVarAllowLiveStreamDelete.GetValueOnAnyThread() || !bIsLive)
 			{
 				UE_CLOG(bIsLive, LogLocalFileReplay, Warning, TEXT("Deleting network replay stream %s that is currently live!"), *StreamName);
 
@@ -1600,7 +1656,7 @@ void FLocalFileNetworkReplayStreamer::RenameReplayFriendlyName_Internal(const FS
 					return;
 				}
 
-				if (SerializationInfo.FileVersion < LocalFileReplay::HISTORY_FIXEDSIZE_FRIENDLY_NAME)
+				if (SerializationInfo.GetLocalFileReplayVersion() < FLocalFileReplayCustomVersion::FixedSizeFriendlyName)
 				{
 					UE_LOG(LogLocalFileReplay, Warning, TEXT("FLocalFileNetworkReplayStreamer::RenameReplayFriendlyName: Replay too old to rename safely %s"), *ReplayName);
 					return;
@@ -2695,7 +2751,7 @@ bool FLocalFileNetworkReplayStreamer::CleanUpOldReplays(const FString& DemoPath,
 			return true;
 		}
 
-		uint64 MinFreeSpace = LocalFileReplay::CVarReplayRecordingMinSpace.GetValueOnAnyThread();
+		uint64 MinFreeSpace = UE::Net::LocalFileReplay::CVarReplayRecordingMinSpace.GetValueOnAnyThread();
 
 		// build an array of replay info sorted by timestamps
 		struct FAutoReplayInfo
@@ -3064,7 +3120,7 @@ void FLocalFileNetworkReplayStreamer::ConditionallyFlushStream()
 		return;
 	}
 
-	const float FLUSH_TIME_IN_SECONDS = LocalFileReplay::CVarChunkUploadDelayInSeconds.GetValueOnGameThread();
+	const float FLUSH_TIME_IN_SECONDS = UE::Net::LocalFileReplay::CVarChunkUploadDelayInSeconds.GetValueOnGameThread();
 
 	if ( FPlatformTime::Seconds() - LastChunkTime > FLUSH_TIME_IN_SECONDS )
 	{
@@ -3103,7 +3159,7 @@ void FLocalFileNetworkReplayStreamer::ConditionallyLoadNextChunk()
 		const double LoadElapsedTime = FPlatformTime::Seconds() - LastChunkTime;
 
 		// Unless it's critical (i.e. bReallyNeedToLoadChunk is true), never try faster than the min delay
-		if (LoadElapsedTime < LocalFileReplay::CVarMinLoadNextChunkDelaySeconds.GetValueOnAnyThread())
+		if (LoadElapsedTime < UE::Net::LocalFileReplay::CVarMinLoadNextChunkDelaySeconds.GetValueOnAnyThread())
 		{
 			return;		
 		}
@@ -3117,7 +3173,7 @@ void FLocalFileNetworkReplayStreamer::ConditionallyLoadNextChunk()
 			const float TimeLeft		= TotalStreamTimeSeconds - CurrentTime;
 
 			// Determine if we have enough buffer to stop streaming for now
-			const float MaxBufferedTimeSeconds = LocalFileReplay::CVarChunkUploadDelayInSeconds.GetValueOnAnyThread() * 0.5f;
+			const float MaxBufferedTimeSeconds = UE::Net::LocalFileReplay::CVarChunkUploadDelayInSeconds.GetValueOnAnyThread() * 0.5f;
 
 			if (TimeLeft > MaxBufferedTimeSeconds)
 			{
@@ -3249,7 +3305,7 @@ void FLocalFileNetworkReplayStreamer::ConditionallyLoadNextChunk()
 				check(StreamTimeRange.IsValid());
 
 				// make space before appending
-				int32 MaxBufferedChunks = LocalFileReplay::CVarMaxBufferedStreamChunks.GetValueOnAnyThread();
+				int32 MaxBufferedChunks = UE::Net::LocalFileReplay::CVarMaxBufferedStreamChunks.GetValueOnAnyThread();
 				if (MaxBufferedChunks > 0)
 				{
 					int32 MinChunkIndex = FMath::Max(0, (RequestedStreamChunkIndex + 1) - MaxBufferedChunks);
@@ -3383,7 +3439,7 @@ void FLocalFileNetworkReplayStreamer::CleanupRequestCache()
 
 		check(OldestKey != INDEX_NONE);
 
-		const uint32 MaxCacheSize = LocalFileReplay::CVarMaxCacheSize.GetValueOnAnyThread();
+		const uint32 MaxCacheSize = UE::Net::LocalFileReplay::CVarMaxCacheSize.GetValueOnAnyThread();
 
 		if (TotalSize <= MaxCacheSize)
 		{
diff --git a/Engine/Source/Runtime/NetworkReplayStreaming/LocalFileNetworkReplayStreaming/Public/LocalFileNetworkReplayStreaming.h b/Engine/Source/Runtime/NetworkReplayStreaming/LocalFileNetworkReplayStreaming/Public/LocalFileNetworkReplayStreaming.h
--- a/Engine/Source/Runtime/NetworkReplayStreaming/LocalFileNetworkReplayStreaming/Public/LocalFileNetworkReplayStreaming.h
+++ b/Engine/Source/Runtime/NetworkReplayStreaming/LocalFileNetworkReplayStreaming/Public/LocalFileNetworkReplayStreaming.h
@@ -8,6 +8,7 @@
 #include "Tickable.h"
 #include "Serialization/ArrayReader.h"
 #include "Serialization/ArrayWriter.h"
+#include "Serialization/CustomVersion.h"
 #include "Serialization/MemoryReader.h"
 #include "Async/Async.h"
 #include "Templates/SharedPointer.h"
@@ -17,6 +18,32 @@
 class FNetworkReplayVersion;
 class FLocalFileNetworkReplayStreamer;
 
+struct LOCALFILENETWORKREPLAYSTREAMING_API FLocalFileReplayCustomVersion
+{
+	enum Type
+	{
+		// Before any version changes were made
+		BeforeCustomVersionWasAdded = 0,
+
+		FixedSizeFriendlyName = 1,
+		CompressionSupport = 2,
+		RecordingTimestamp = 3,
+		StreamChunkTimes = 4,
+		FriendlyNameCharEncoding = 5,
+		EncryptionSupport = 6,
+		CustomVersions = 7,
+
+		// -----<new versions can be added above this line>-------------------------------------------------
+		VersionPlusOne,
+		LatestVersion = VersionPlusOne - 1
+	};
+
+	// The GUID for this custom version number
+	const static FGuid Guid;
+
+	FLocalFileReplayCustomVersion() = delete;
+};
+
 enum class ELocalFileChunkType : uint32
 {
 	Header,
@@ -629,8 +656,13 @@ protected:
 	{
 		FLocalFileSerializationInfo();
 
+		FLocalFileReplayCustomVersion::Type GetLocalFileReplayVersion() const;
+
+		UE_DEPRECATED(5.2, "Replaced by FileCustomVersions.")
 		uint32 FileVersion;
+
 		FString FileFriendlyName;
+		FCustomVersionContainer FileCustomVersions;
 	};
 
 	bool ReadReplayInfo(const FString& StreamName, FLocalFileReplayInfo& OutReplayInfo) const
@@ -703,10 +735,10 @@ public:
 	static bool CleanUpOldReplays(const FString& DemoPath = GetDefaultDemoSavePath(), TArrayView<const FString> AdditionalRelativeDemoPaths = {});
 	static bool GetDemoFreeStorageSpace(uint64& DiskFreeSpace, const FString& DemoPath);
 
-
-
 	static const uint32 FileMagic;
 	static const uint32 MaxFriendlyNameLen;
+
+	UE_DEPRECATED(5.2, "No longer used, replaced with custom version.")
 	static const uint32 LatestVersion;
 };
 
