Work around static analysis false positive in NullNetworkReplayStreaming with clang 15.

diff --git a/Engine/Source/Runtime/NetworkReplayStreaming/NullNetworkReplayStreaming/Private/NullNetworkReplayStreaming.cpp b/Engine/Source/Runtime/NetworkReplayStreaming/NullNetworkReplayStreaming/Private/NullNetworkReplayStreaming.cpp
--- a/Engine/Source/Runtime/NetworkReplayStreaming/NullNetworkReplayStreaming/Private/NullNetworkReplayStreaming.cpp
+++ b/Engine/Source/Runtime/NetworkReplayStreaming/NullNetworkReplayStreaming/Private/NullNetworkReplayStreaming.cpp
@@ -752,8 +752,6 @@ EStreamingOperationResult FNullNetworkReplayStreamer::GetDemoPath(FString& DemoP
 	return EStreamingOperationResult::Success;
 }
 
-IMPLEMENT_MODULE( FNullNetworkReplayStreamingFactory, NullNetworkReplayStreaming )
-
 TSharedPtr< INetworkReplayStreamer > FNullNetworkReplayStreamingFactory::CreateReplayStreamer() 
 {
 	return TSharedPtr< INetworkReplayStreamer >( new FNullNetworkReplayStreamer );
diff --git a/Engine/Source/Runtime/NetworkReplayStreaming/NullNetworkReplayStreaming/Private/NullNetworkReplayStreamingModule.cpp b/Engine/Source/Runtime/NetworkReplayStreaming/NullNetworkReplayStreaming/Private/NullNetworkReplayStreamingModule.cpp
new file mode 100644
--- /dev/null
+++ b/Engine/Source/Runtime/NetworkReplayStreaming/NullNetworkReplayStreaming/Private/NullNetworkReplayStreamingModule.cpp
@@ -0,0 +1,6 @@
+// Copyright Epic Games, Inc. All Rights Reserved.
+
+#include "NullNetworkReplayStreaming.h"
+#include "Modules/ModuleManager.h"
+
+IMPLEMENT_MODULE(FNullNetworkReplayStreamingFactory, NullNetworkReplayStreaming)
