Refactoring how current replay info is accessed and updated to avoid extra file reads during recording

diff --git a/Engine/Source/Runtime/Engine/Private/DemoNetDriver.cpp b/Engine/Source/Runtime/Engine/Private/DemoNetDriver.cpp
--- a/Engine/Source/Runtime/Engine/Private/DemoNetDriver.cpp
+++ b/Engine/Source/Runtime/Engine/Private/DemoNetDriver.cpp
@@ -223,6 +223,36 @@ namespace DemoNetDriverRecordingPrivate
 					}
 				}
 			}));
+
+#if !UE_BUILD_SHIPPING
+	static FAutoConsoleCommandWithWorldAndArgs DemoTestWriteEvent(
+		TEXT("Demo.TestWriteEvent"),
+		TEXT("Add or update a test replay event on the currently recording replay, with an optional argument for event size in bytes"),
+		FConsoleCommandWithWorldAndArgsDelegate::CreateStatic(
+			[](const TArray<FString>& Params, UWorld* World)
+			{
+				if (World)
+				{
+					if (UDemoNetDriver* Driver = World->GetDemoNetDriver())
+					{
+						if (APlayerController* ViewerPC = GEngine->GetFirstLocalPlayerController(World))
+						{
+							int32 EventSize = 32;
+
+							if (Params.Num() > 0)
+							{
+								EventSize = FCString::Atoi(*Params[0]);
+							}
+
+							TArray<uint8> EventData;
+							EventData.AddZeroed(EventSize);
+
+							Driver->AddOrUpdateEvent(TEXT("Test"), TEXT("Test"), TEXT("Test"), EventData);
+						}
+					}
+				}
+			}));
+#endif
 }
 
 struct FDemoBudgetLogHelper
diff --git a/Engine/Source/Runtime/NetworkReplayStreaming/LocalFileNetworkReplayStreaming/Private/LocalFileNetworkReplayStreaming.cpp b/Engine/Source/Runtime/NetworkReplayStreaming/LocalFileNetworkReplayStreaming/Private/LocalFileNetworkReplayStreaming.cpp
--- a/Engine/Source/Runtime/NetworkReplayStreaming/LocalFileNetworkReplayStreaming/Private/LocalFileNetworkReplayStreaming.cpp
+++ b/Engine/Source/Runtime/NetworkReplayStreaming/LocalFileNetworkReplayStreaming/Private/LocalFileNetworkReplayStreaming.cpp
@@ -11,6 +11,7 @@
 #include "Engine/LocalPlayer.h"
 #include "Engine/World.h"
 #include "Online/CoreOnline.h"
+#include "Serialization/ArchiveCountMem.h"
 #include "Serialization/LargeMemoryReader.h"
 #include "Misc/ConfigCacheIni.h"
 #include "HAL/IPlatformFileManagedStorageWrapper.h"
@@ -90,6 +91,34 @@ FLocalFileReplayCustomVersion::Type FLocalFileNetworkReplayStreamer::FLocalFileS
 	PRAGMA_ENABLE_DEPRECATION_WARNINGS
 }
 
+void FLocalFileEventInfo::CountBytes(FArchive& Ar) const
+{
+	Id.CountBytes(Ar);
+	Group.CountBytes(Ar);
+	Metadata.CountBytes(Ar);
+}
+
+void FLocalFileReplayInfo::CountBytes(FArchive& Ar) const
+{
+	FriendlyName.CountBytes(Ar);
+	EncryptionKey.CountBytes(Ar);
+	
+	Chunks.CountBytes(Ar);
+	DataChunks.CountBytes(Ar);
+
+	Checkpoints.CountBytes(Ar);
+	for (const FLocalFileEventInfo& Info : Checkpoints)
+	{
+		Info.CountBytes(Ar);
+	}
+
+	Events.CountBytes(Ar);
+	for (const FLocalFileEventInfo& Info : Events)
+	{
+		Info.CountBytes(Ar);
+	}
+}
+
 void FLocalFileStreamFArchive::Serialize(void* V, int64 Length) 
 {
 	if (IsLoading())
@@ -544,6 +573,12 @@ bool FLocalFileNetworkReplayStreamer::ReadReplayInfo(FArchive& Archive, FLocalFi
 
 		Info.bIsValid = EnumHasAnyFlags(Flags, EReadReplayInfoFlags::SkipHeaderChunkTest) || Info.Chunks.IsValidIndex(Info.HeaderChunkIndex);
 
+		FArchiveCountMem CountMemAr(nullptr);
+		Info.CountBytes(CountMemAr);
+		const int64 InfoSize = sizeof(FLocalFileReplayInfo) + CountMemAr.GetNum();
+
+		UE_LOG(LogLocalFileReplay, Log, TEXT("ReadReplayInfo: IsValid: %s MemSize: %lld bytes"), *LexToString(Info.bIsValid), InfoSize);
+
 		return Info.bIsValid && !Archive.IsError();
 	}
 
@@ -765,7 +800,7 @@ void FLocalFileNetworkReplayStreamer::StartStreaming(const FStartStreamingParame
 				else
 				{
 					// Load metadata if it exists
-					ReadReplayInfo(CurrentStreamName, RequestData.ReplayInfo);
+					ReadReplayInfo(CurrentStreamName, TaskReplayInfo);
 				}
 			},
 			[this, Delegate](TLocalFileRequestCommonData<FStartStreamingResult>& RequestData)
@@ -776,7 +811,7 @@ void FLocalFileNetworkReplayStreamer::StartStreaming(const FStartStreamingParame
 				}
 				else
 				{
-					CurrentReplayInfo = RequestData.ReplayInfo;
+					UpdateCurrentReplayInfo(TaskReplayInfo, EUpdateReplayInfoFlags::FullUpdate);
 
 					if (!CurrentReplayInfo.bIsValid)
 					{
@@ -842,6 +877,7 @@ void FLocalFileNetworkReplayStreamer::StartStreaming(const FStartStreamingParame
 				if ((FinalDemoName.IsEmpty() || !FullDemoFilename.EndsWith(FNetworkReplayStreaming::GetReplayFileExtension())))
 				{
 					UE_LOG(LogLocalFileReplay, Warning, TEXT("FLocalFileNetworkReplayStreamer::StartStreaming. Invalid replay file name for recording: %s"), *FullDemoFilename);
+					RequestData.AsyncError = ELocalFileReplayResult::InvalidName;
 					return;
 				}
 
@@ -859,20 +895,23 @@ void FLocalFileNetworkReplayStreamer::StartStreaming(const FStartStreamingParame
 				// Delete any existing demo with this name
 				IFileManager::Get().Delete(*FullDemoFilename);
 
-				RequestData.ReplayInfo.NetworkVersion = Params.ReplayVersion.NetworkVersion;
-				RequestData.ReplayInfo.Changelist = Params.ReplayVersion.Changelist;
-				RequestData.ReplayInfo.FriendlyName = Params.FriendlyName;
-				RequestData.ReplayInfo.bIsLive = true;
-				RequestData.ReplayInfo.Timestamp = FDateTime::Now();
-				RequestData.ReplayInfo.EncryptionKey = EncryptionKey;
+				TaskReplayInfo.NetworkVersion = Params.ReplayVersion.NetworkVersion;
+				TaskReplayInfo.Changelist = Params.ReplayVersion.Changelist;
+				TaskReplayInfo.FriendlyName = Params.FriendlyName;
+				TaskReplayInfo.bIsLive = true;
+				TaskReplayInfo.Timestamp = FDateTime::Now();
+				TaskReplayInfo.EncryptionKey = EncryptionKey;
 
-				WriteReplayInfo(CurrentStreamName, RequestData.ReplayInfo);
+				TaskReplayInfo.bIsValid = WriteReplayInfo(CurrentStreamName, TaskReplayInfo);
 
 				RequestData.DelegateResult.Result = EStreamingOperationResult::Success;
 			},
 			[this, Delegate](TLocalFileRequestCommonData<FStartStreamingResult>& RequestData)
 			{
-				CurrentReplayInfo = RequestData.ReplayInfo;
+				if (RequestData.AsyncError == ELocalFileReplayResult::Success)
+				{
+					UpdateCurrentReplayInfo(TaskReplayInfo, EUpdateReplayInfoFlags::FullUpdate);
+				}
 
 				Delegate.ExecuteIfBound(RequestData.DelegateResult);
 			});
@@ -936,22 +975,19 @@ void FLocalFileNetworkReplayStreamer::StopStreaming()
 
 		FlushStream(TotalLengthInMS);
 
-		AddGenericRequestToQueue<FLocalFileReplayInfo>(EQueuedLocalFileRequestType::StopRecording,
-			[this, TotalLengthInMS](FLocalFileReplayInfo& ReplayInfo)
-			{
-				if (ReadReplayInfo(CurrentStreamName, ReplayInfo))
+		AddGenericRequestToQueue<ELocalFileReplayResult>(EQueuedLocalFileRequestType::StopRecording,
+			[this, TotalLengthInMS](ELocalFileReplayResult& ReplayResult)
 				{
 					// Set the final values of these header properties
-					ReplayInfo.bIsLive = false;
-					ReplayInfo.LengthInMS = TotalLengthInMS;
-					ReplayInfo.EncryptionKey = CurrentReplayInfo.EncryptionKey;
+				TaskReplayInfo.bIsLive = false;
+				TaskReplayInfo.LengthInMS = TotalLengthInMS;
+				TaskReplayInfo.EncryptionKey = CurrentReplayInfo.EncryptionKey;
 
-					WriteReplayInfo(CurrentStreamName, ReplayInfo);
-				}
+				WriteReplayInfo(CurrentStreamName, TaskReplayInfo);
 			},
-			[this](FLocalFileReplayInfo& ReplayInfo)
+			[this](ELocalFileReplayResult& ReplayResult)
 			{
-				CurrentReplayInfo = ReplayInfo;
+				UpdateCurrentReplayInfo(TaskReplayInfo, EUpdateReplayInfoFlags::FullUpdate);
 			});
 	}
 
@@ -1162,8 +1198,8 @@ void FLocalFileNetworkReplayStreamer::AddOrUpdateEvent(const FString& Name, cons
 
 	UE_LOG(LogLocalFileReplay, Verbose, TEXT("FLocalFileNetworkReplayStreamer::AddOrUpdateEvent. Size: %i"), Data.Num());
 
-	AddGenericRequestToQueue<FLocalFileReplayInfo>(EQueuedLocalFileRequestType::UpdatingEvent,
-		[this, Name, Group, TimeInMS, Meta, Data, EncryptionKey=CurrentReplayInfo.EncryptionKey](FLocalFileReplayInfo& ReplayInfo)
+	AddGenericRequestToQueue<ELocalFileReplayResult>(EQueuedLocalFileRequestType::UpdatingEvent,
+		[this, Name, Group, TimeInMS, Meta, Data](ELocalFileReplayResult& ReplayResult)
 		{
 			SCOPE_CYCLE_COUNTER(STAT_LocalReplay_FlushEvent);
 
@@ -1178,17 +1214,15 @@ void FLocalFileNetworkReplayStreamer::AddOrUpdateEvent(const FString& Name, cons
 			// prefix with stream name to be consistent with http streamer
 			EventName = CurrentStreamName + TEXT("_") + EventName;
 
-			if (ReadReplayInfo(CurrentStreamName, ReplayInfo, EReadReplayInfoFlags::SkipHeaderChunkTest))
-			{
 				TSharedPtr<FArchive> LocalFileAr = CreateLocalFileWriter(GetDemoFullFilename(CurrentStreamName));
 				if (LocalFileAr.IsValid())
 				{
 					int32 EventIndex = INDEX_NONE;
 
 					// see if this event already exists
-					for (int32 i=0; i < ReplayInfo.Events.Num(); ++i)
+				for (int32 i=0; i < TaskReplayInfo.Events.Num(); ++i)
 					{
-						if (ReplayInfo.Events[i].Id == EventName)
+					if (TaskReplayInfo.Events[i].Id == EventName)
 						{
 							EventIndex = i;
 							break;
@@ -1201,10 +1235,10 @@ void FLocalFileNetworkReplayStreamer::AddOrUpdateEvent(const FString& Name, cons
 					{
 						SCOPE_CYCLE_COUNTER(STAT_LocalReplay_EncryptTime);
 
-						if (!EncryptBuffer(Data, EncryptedData, EncryptionKey))
+					if (!EncryptBuffer(Data, EncryptedData, TaskReplayInfo.EncryptionKey))
 						{
 							UE_LOG(LogLocalFileReplay, Warning, TEXT("FLocalFileNetworkReplayStreamer::AddOrUpdateEvent - EncryptBuffer failed"));
-							ReplayInfo.bIsValid = false;
+						ReplayResult = ELocalFileReplayResult::EncryptBuffer;
 							return;
 						}
 					}
@@ -1219,21 +1253,21 @@ void FLocalFileNetworkReplayStreamer::AddOrUpdateEvent(const FString& Name, cons
 					ELocalFileChunkType ChunkType = ELocalFileChunkType::Event;
 					Writer << ChunkType;
 
-					int64 SavedPos = Writer.Tell();
+				const int64 SavedPos = Writer.Tell();
 
 					int32 PlaceholderSize = 0;
 					Writer << PlaceholderSize;
 
-					int64 MetadataPos = Writer.Tell();
+				const int64 MetadataPos = Writer.Tell();
 
-					FString TempName = EventName;
-					Writer << TempName;
+				FString TempId = EventName;
+				Writer << TempId;
 
-					FString Value = Group;
-					Writer << Value;
+				FString GroupValue = Group;
+				Writer << GroupValue;
 
-					Value = Meta;
-					Writer << Value;
+				FString MetaValue = Meta;
+				Writer << MetaValue;
 
 					uint32 Time1 = TimeInMS;
 					Writer << Time1;
@@ -1244,12 +1278,13 @@ void FLocalFileNetworkReplayStreamer::AddOrUpdateEvent(const FString& Name, cons
 					int32 EventSize = EncryptedData.Num();
 					Writer << EventSize;
 
+				const int64 InternalDataOffset = Writer.Tell();
+
 					Writer.Serialize((void*)EncryptedData.GetData(), EncryptedData.Num());
 
 					int32 ChunkSize = Writer.Tell() - MetadataPos;
 
-					Writer.Seek(SavedPos);
-					Writer << ChunkSize;
+				bool bNewChunk = true;
 
 					if (EventIndex == INDEX_NONE)
 					{
@@ -1258,35 +1293,75 @@ void FLocalFileNetworkReplayStreamer::AddOrUpdateEvent(const FString& Name, cons
 					}
 					else 
 					{
-						if (ChunkSize > ReplayInfo.Chunks[ReplayInfo.Events[EventIndex].ChunkIndex].SizeInBytes)
+					if (ChunkSize > TaskReplayInfo.Chunks[TaskReplayInfo.Events[EventIndex].ChunkIndex].SizeInBytes)
 						{
-							LocalFileAr->Seek(ReplayInfo.Chunks[ReplayInfo.Events[EventIndex].ChunkIndex].TypeOffset);
+						LocalFileAr->Seek(TaskReplayInfo.Chunks[TaskReplayInfo.Events[EventIndex].ChunkIndex].TypeOffset);
 
 							// clear chunk type so it will be skipped later
 							ChunkType = ELocalFileChunkType::Unknown;
 							*LocalFileAr << ChunkType;
 
+						TaskReplayInfo.Chunks[TaskReplayInfo.Events[EventIndex].ChunkIndex].ChunkType = ELocalFileChunkType::Unknown;
+
 							LocalFileAr->Seek(LocalFileAr->TotalSize());
 						}
 						else
 						{
-							LocalFileAr->Seek(ReplayInfo.Chunks[ReplayInfo.Events[EventIndex].ChunkIndex].TypeOffset);
+						bNewChunk = false;
+
+						// reuse existing chunk
+						LocalFileAr->Seek(TaskReplayInfo.Chunks[TaskReplayInfo.Events[EventIndex].ChunkIndex].TypeOffset);
+
+						// maintain the original chunk size to avoid corrupting the file
+						ChunkSize = TaskReplayInfo.Chunks[TaskReplayInfo.Events[EventIndex].ChunkIndex].SizeInBytes;
+					}
 						}
+
+				Writer.Seek(SavedPos);
+				Writer << ChunkSize;
+
+				const int64 TypeOffset = LocalFileAr->Tell();
+
+				if (bNewChunk)
+				{
+					FLocalFileChunkInfo& ChunkInfo = TaskReplayInfo.Chunks.AddDefaulted_GetRef();
+					ChunkInfo.ChunkType = ChunkType;
+					ChunkInfo.TypeOffset = TypeOffset;
+					ChunkInfo.DataOffset = TypeOffset + MetadataPos;
+					ChunkInfo.SizeInBytes = ChunkSize;
+
+					FLocalFileEventInfo& NewEventInfo = TaskReplayInfo.Events.AddDefaulted_GetRef();
+					NewEventInfo.ChunkIndex = TaskReplayInfo.Chunks.Num() - 1;
+					NewEventInfo.Id = MoveTemp(TempId);
+
+					EventIndex = TaskReplayInfo.Events.Num() - 1;
 					}
 
-					LocalFileAr->Serialize(Writer.GetData(), Writer.TotalSize());
+				FLocalFileEventInfo& EventInfo = TaskReplayInfo.Events[EventIndex];
+				EventInfo.Group = MoveTemp(GroupValue);
+				EventInfo.Metadata = MoveTemp(MetaValue);
+				EventInfo.Time1 = Time1;
+				EventInfo.Time2 = Time2;
+				EventInfo.SizeInBytes = EventSize;
+				EventInfo.EventDataOffset = TypeOffset + InternalDataOffset;
 
+					LocalFileAr->Serialize(Writer.GetData(), Writer.TotalSize());
 					LocalFileAr = nullptr;
 				}
-
-				ReadReplayInfo(CurrentStreamName, ReplayInfo);
+			else
+			{
+				ReplayResult = ELocalFileReplayResult::FileWriter;
 			}
 		},
-		[this](FLocalFileReplayInfo& ReplayInfo)
+		[this](ELocalFileReplayResult& ReplayResult)
 		{
-			if (ReplayInfo.bIsValid)
+			if (ReplayResult == ELocalFileReplayResult::Success)
+			{
+				UpdateCurrentReplayInfo(TaskReplayInfo, EUpdateReplayInfoFlags::None);
+			}
+			else
 			{
-				UpdateCurrentReplayInfo(ReplayInfo);
+				SetLastError(ReplayResult);
 			}
 		});
 }
@@ -1755,11 +1830,13 @@ void FLocalFileNetworkReplayStreamer::FlushStream(const uint32 TimeInMS)
 
 	StreamTimeRange.Min = StreamTimeRange.Max;
 
+	const int32 TotalLengthInMS = CurrentReplayInfo.LengthInMS;
+
 	// Save any newly streamed data to disk
 	UE_LOG(LogLocalFileReplay, Verbose, TEXT("FLocalFileNetworkReplayStreamer::FlushStream. StreamChunkIndex: %i, Size: %i"), StreamChunkIndex, StreamAr.Buffer.Num());
 
-	AddGenericRequestToQueue<FLocalFileReplayInfo>(EQueuedLocalFileRequestType::WritingStream,
-		[this, StreamChunkStartMS, StreamChunkEndMS, StreamData = MoveTemp(StreamAr.Buffer), EncryptionKey = CurrentReplayInfo.EncryptionKey](FLocalFileReplayInfo& ReplayInfo) mutable
+	AddGenericRequestToQueue<ELocalFileReplayResult>(EQueuedLocalFileRequestType::WritingStream,
+		[this, StreamChunkStartMS, StreamChunkEndMS, TotalLengthInMS, StreamData = MoveTemp(StreamAr.Buffer)](ELocalFileReplayResult& ReplayResult) mutable
 		{
 			SCOPE_CYCLE_COUNTER(STAT_LocalReplay_FlushStream);
 
@@ -1779,7 +1856,7 @@ void FLocalFileNetworkReplayStreamer::FlushStream(const uint32 TimeInMS)
 					if (!CompressBuffer(StreamData, CompressedData))
 					{
 						UE_LOG(LogLocalFileReplay, Warning, TEXT("FLocalFileNetworkReplayStreamer::FlushStream - CompressBuffer failed"));
-						ReplayInfo.bIsValid = false;
+						ReplayResult = ELocalFileReplayResult::CompressBuffer;
 						return;
 					}
 				}
@@ -1794,10 +1871,10 @@ void FLocalFileNetworkReplayStreamer::FlushStream(const uint32 TimeInMS)
 				{
 					SCOPE_CYCLE_COUNTER(STAT_LocalReplay_EncryptTime);
 
-					if (!EncryptBuffer(CompressedData, EncryptedData, EncryptionKey))
+					if (!EncryptBuffer(CompressedData, EncryptedData, TaskReplayInfo.EncryptionKey))
 					{
 						UE_LOG(LogLocalFileReplay, Warning, TEXT("FLocalFileNetworkReplayStreamer::FlushStream - EncryptBuffer failed"));
-						ReplayInfo.bIsValid = false;
+						ReplayResult = ELocalFileReplayResult::EncryptBuffer;
 						return;
 					}
 				}
@@ -1809,15 +1886,17 @@ void FLocalFileNetworkReplayStreamer::FlushStream(const uint32 TimeInMS)
 				// flush chunk to disk
 				if (EncryptedData.Num() > 0)
 				{
+					const int64 TypeOffset = LocalFileAr->Tell();
+
 					ELocalFileChunkType ChunkType = ELocalFileChunkType::ReplayData;
 					*LocalFileAr << ChunkType;
 
-					int64 SavedPos = LocalFileAr->Tell();
+					const int64 SavedPos = LocalFileAr->Tell();
 
 					int32 PlaceholderSize = 0;
 					*LocalFileAr << PlaceholderSize;
 
-					int64 MetadataPos = LocalFileAr->Tell();
+					const int64 MetadataPos = LocalFileAr->Tell();
 
 					uint32 Time1 = StreamChunkStartMS;
 					*LocalFileAr << Time1;
@@ -1830,29 +1909,54 @@ void FLocalFileNetworkReplayStreamer::FlushStream(const uint32 TimeInMS)
 
 					*LocalFileAr << SizeInMemory;
 
+					const int64 ReplayDataOffset = LocalFileAr->Tell();
+
 					LocalFileAr->Serialize((void*)EncryptedData.GetData(), EncryptedData.Num());
 
 					int32 ChunkSize = LocalFileAr->Tell() - MetadataPos;
 
 					LocalFileAr->Seek(SavedPos);
 					*LocalFileAr << ChunkSize;
+
+					FLocalFileChunkInfo& ChunkInfo = TaskReplayInfo.Chunks.AddDefaulted_GetRef();
+					ChunkInfo.ChunkType = ChunkType;
+					ChunkInfo.TypeOffset = TypeOffset;
+					ChunkInfo.DataOffset = MetadataPos;
+					ChunkInfo.SizeInBytes = ChunkSize;
+
+					FLocalFileReplayDataInfo& DataInfo = TaskReplayInfo.DataChunks.AddDefaulted_GetRef();
+					DataInfo.ChunkIndex = TaskReplayInfo.Chunks.Num() - 1;
+					DataInfo.Time1 = Time1;
+					DataInfo.Time2 = Time2;
+					DataInfo.SizeInBytes = SizeOnDisk;
+					DataInfo.MemorySizeInBytes = SizeInMemory;
+					DataInfo.ReplayDataOffset = ReplayDataOffset;
+					DataInfo.StreamOffset = TaskReplayInfo.TotalDataSizeInBytes;
+					
+					TaskReplayInfo.TotalDataSizeInBytes += DataInfo.MemorySizeInBytes;
 				}
 
 				LocalFileAr = nullptr;
-			}
 
-			ReadReplayInfo(CurrentStreamName, ReplayInfo);
+				TaskReplayInfo.LengthInMS = TotalLengthInMS;
+
+				WriteReplayInfo(CurrentStreamName, TaskReplayInfo);
+			}
+			else
+			{
+				ReplayResult = ELocalFileReplayResult::FileWriter;
+			}
 		},
-		[this](FLocalFileReplayInfo& ReplayInfo)
+		[this](ELocalFileReplayResult& ReplayResult)
 		{
-			if (ReplayInfo.bIsValid)
+			if (ReplayResult == ELocalFileReplayResult::Success)
 			{
-				UpdateCurrentReplayInfo(ReplayInfo);
+				UpdateCurrentReplayInfo(TaskReplayInfo, EUpdateReplayInfoFlags::None);
 			}
 			else
 			{
-				UE_LOG(LogLocalFileReplay, Error, TEXT("FLocalFileNetworkReplayStreamer::FlushStream failed with invalid replay info."));
-				SetLastError(ELocalFileReplayResult::InvalidReplayInfo);
+				UE_LOG(LogLocalFileReplay, Error, TEXT("FLocalFileNetworkReplayStreamer::FlushStream failed."));
+				SetLastError(ReplayResult);
 			}
 		});
 
@@ -1894,15 +1998,13 @@ void FLocalFileNetworkReplayStreamer::FlushCheckpointInternal(const uint32 TimeI
 	const int32 TotalLengthInMS = CurrentReplayInfo.LengthInMS;
 	const uint32 CheckpointTimeInMS = StreamTimeRange.Max;
 
-	AddGenericRequestToQueue<FLocalFileReplayInfo>(EQueuedLocalFileRequestType::WritingCheckpoint, 
-		[this, CheckpointTimeInMS, TotalLengthInMS, CheckpointData=MoveTemp(CheckpointAr.Buffer), EncryptionKey=CurrentReplayInfo.EncryptionKey](FLocalFileReplayInfo& ReplayInfo) mutable
+	AddGenericRequestToQueue<ELocalFileReplayResult>(EQueuedLocalFileRequestType::WritingCheckpoint, 
+		[this, CheckpointTimeInMS, TotalLengthInMS, CheckpointData=MoveTemp(CheckpointAr.Buffer)](ELocalFileReplayResult& ReplayResult) mutable
 		{
 			SCOPE_CYCLE_COUNTER(STAT_LocalReplay_FlushCheckpoint);
 
-			if (ReadReplayInfo(CurrentStreamName, ReplayInfo))
-			{
-				int32 DataChunkIndex = ReplayInfo.DataChunks.Num();
-				int32 CheckpointIndex = ReplayInfo.Checkpoints.Num();
+			const int32 DataChunkIndex = TaskReplayInfo.DataChunks.Num();
+			const int32 CheckpointIndex = TaskReplayInfo.Checkpoints.Num();
 
 				TSharedPtr<FArchive> LocalFileAr = CreateLocalFileWriter(GetDemoFullFilename(CurrentStreamName));
 				if (LocalFileAr.IsValid())
@@ -1918,7 +2020,7 @@ void FLocalFileNetworkReplayStreamer::FlushCheckpointInternal(const uint32 TimeI
 						if (!CompressBuffer(CheckpointData, CompressedData))
 						{
 							UE_LOG(LogLocalFileReplay, Warning, TEXT("FLocalFileNetworkReplayStreamer::FlushStream - CompressBuffer failed"));
-							ReplayInfo.bIsValid = false;
+						ReplayResult = ELocalFileReplayResult::CompressBuffer;
 							return;
 						}
 					}
@@ -1933,10 +2035,10 @@ void FLocalFileNetworkReplayStreamer::FlushCheckpointInternal(const uint32 TimeI
 					{
 						SCOPE_CYCLE_COUNTER(STAT_LocalReplay_EncryptTime);
 
-						if (!EncryptBuffer(CompressedData, EncryptedData, EncryptionKey))
+					if (!EncryptBuffer(CompressedData, EncryptedData, TaskReplayInfo.EncryptionKey))
 						{
 							UE_LOG(LogLocalFileReplay, Warning, TEXT("FLocalFileNetworkReplayStreamer::FlushStream - EncryptBuffer failed"));
-							ReplayInfo.bIsValid = false;
+						ReplayResult = ELocalFileReplayResult::EncryptBuffer;
 							return;
 						}
 					}
@@ -1948,15 +2050,17 @@ void FLocalFileNetworkReplayStreamer::FlushCheckpointInternal(const uint32 TimeI
 					// flush checkpoint
 					if (EncryptedData.Num() > 0)
 					{
+					const int64 TypeOffset = LocalFileAr->Tell();
+
 						ELocalFileChunkType ChunkType = ELocalFileChunkType::Checkpoint;
 						*LocalFileAr << ChunkType;
 
-						int64 SavedPos = LocalFileAr->Tell();
+					const int64 SavedPos = LocalFileAr->Tell();
 
 						int32 PlaceholderSize = 0;
 						*LocalFileAr << PlaceholderSize;
 
-						int64 MetadataPos = LocalFileAr->Tell();
+					const int64 MetadataPos = LocalFileAr->Tell();
 
 						FString Id = FString::Printf(TEXT("checkpoint%ld"), CheckpointIndex);
 						*LocalFileAr << Id;
@@ -1976,36 +2080,54 @@ void FLocalFileNetworkReplayStreamer::FlushCheckpointInternal(const uint32 TimeI
 						int32 CheckpointSize = EncryptedData.Num();
 						*LocalFileAr << CheckpointSize;
 
+					const int64 EventDataOffset = LocalFileAr->Tell();
+
 						LocalFileAr->Serialize((void*)EncryptedData.GetData(), EncryptedData.Num());
 
 						int32 ChunkSize = LocalFileAr->Tell() - MetadataPos;
 
 						LocalFileAr->Seek(SavedPos);
 						*LocalFileAr << ChunkSize;
+
+					FLocalFileChunkInfo& ChunkInfo = TaskReplayInfo.Chunks.AddDefaulted_GetRef();
+					ChunkInfo.ChunkType = ChunkType;
+					ChunkInfo.TypeOffset = TypeOffset;
+					ChunkInfo.DataOffset = MetadataPos;
+					ChunkInfo.SizeInBytes = ChunkSize;
+
+					FLocalFileEventInfo& CheckpointInfo = TaskReplayInfo.Checkpoints.AddDefaulted_GetRef();
+					CheckpointInfo.ChunkIndex = TaskReplayInfo.Chunks.Num() - 1;
+					CheckpointInfo.Id = MoveTemp(Id);
+					CheckpointInfo.Group = MoveTemp(Group);
+					CheckpointInfo.Metadata = MoveTemp(Metadata);
+					CheckpointInfo.Time1 = Time1;
+					CheckpointInfo.Time2 = Time2;
+					CheckpointInfo.SizeInBytes = CheckpointSize;
+					CheckpointInfo.EventDataOffset = EventDataOffset;
 					}
 
 					LocalFileAr = nullptr;
 				}
+			else
+			{
+				ReplayResult = ELocalFileReplayResult::FileWriter;
+				return;
+			}
 
-				if (ReadReplayInfo(CurrentStreamName, ReplayInfo))
-				{
-					ReplayInfo.LengthInMS = TotalLengthInMS;
-					ReplayInfo.EncryptionKey = EncryptionKey;
+			TaskReplayInfo.LengthInMS = TotalLengthInMS;
 
-					WriteReplayInfo(CurrentStreamName, ReplayInfo);
-				}
-			}
+			WriteReplayInfo(CurrentStreamName, TaskReplayInfo);
 		},
-		[this](FLocalFileReplayInfo& ReplayInfo)
+		[this](ELocalFileReplayResult& ReplayResult)
 		{
-			if (ReplayInfo.bIsValid)
+			if (ReplayResult == ELocalFileReplayResult::Success)
 			{
-				UpdateCurrentReplayInfo(ReplayInfo);
+				UpdateCurrentReplayInfo(TaskReplayInfo, EUpdateReplayInfoFlags::None);
 			}
 			else
 			{
-				UE_LOG(LogLocalFileReplay, Error, TEXT("FLocalFileNetworkReplayStreamer::FlushCheckpointInternal failed with invalid replay info."));
-				SetLastError(ELocalFileReplayResult::InvalidReplayInfo);
+				UE_LOG(LogLocalFileReplay, Error, TEXT("FLocalFileNetworkReplayStreamer::FlushCheckpointInternal failed."));
+				SetLastError(ReplayResult);
 			}
 		});
 
@@ -2097,8 +2219,6 @@ void FLocalFileNetworkReplayStreamer::GotoCheckpointIndex(const int32 Checkpoint
 			TSharedPtr<FArchive> LocalFileAr = CreateLocalFileReader(FullDemoFilename);
 			if (LocalFileAr.IsValid())
 			{
-				if (ReadReplayInfo(*LocalFileAr, RequestData.ReplayInfo, EReadReplayInfoFlags::None))
-				{
 					TArray<uint8> CheckpointData;
 
 					// read all the checkpoints
@@ -2114,23 +2234,23 @@ void FLocalFileNetworkReplayStreamer::GotoCheckpointIndex(const int32 Checkpoint
 						}
 						else
 						{
-							LocalFileAr->Seek(RequestData.ReplayInfo.Checkpoints[i].EventDataOffset);
+							LocalFileAr->Seek(TaskReplayInfo.Checkpoints[i].EventDataOffset);
 
 							CheckpointData.Reset();
-							CheckpointData.AddUninitialized(RequestData.ReplayInfo.Checkpoints[i].SizeInBytes);
+							CheckpointData.AddUninitialized(TaskReplayInfo.Checkpoints[i].SizeInBytes);
 
 							LocalFileAr->Serialize(CheckpointData.GetData(), CheckpointData.Num());
 
 							TArray<uint8> PlaintextData;
 
 							// Get the checkpoint data
-							if (RequestData.ReplayInfo.bEncrypted)
+							if (TaskReplayInfo.bEncrypted)
 							{
 								if (SupportsEncryption())
 								{
 									SCOPE_CYCLE_COUNTER(STAT_LocalReplay_DecryptTime);
 
-									if (!DecryptBuffer(CheckpointData, PlaintextData, RequestData.ReplayInfo.EncryptionKey))
+									if (!DecryptBuffer(CheckpointData, PlaintextData, TaskReplayInfo.EncryptionKey))
 									{
 										UE_LOG(LogLocalFileReplay, Error, TEXT("FLocalFileNetworkReplayStreamer::GotoCheckpointIndexDelta. DecryptBuffer FAILED."));
 										RequestData.DataBuffer.Empty();
@@ -2151,7 +2271,7 @@ void FLocalFileNetworkReplayStreamer::GotoCheckpointIndex(const int32 Checkpoint
 
 							TArray<uint8> UncompressedData;
 
-							if (RequestData.ReplayInfo.bCompressed)
+							if (TaskReplayInfo.bCompressed)
 							{
 								if (SupportsCompression())
 								{
@@ -2185,7 +2305,6 @@ void FLocalFileNetworkReplayStreamer::GotoCheckpointIndex(const int32 Checkpoint
 							DeltaCheckpointCache.Add(i, MakeShareable(new FCachedFileRequest(UncompressedData, 0)));
 						}
 					}
-				}
 
 				LocalFileAr = nullptr;
 			}
@@ -2302,17 +2421,15 @@ void FLocalFileNetworkReplayStreamer::GotoCheckpointIndex(const int32 Checkpoint
 			TSharedPtr<FArchive> LocalFileAr = CreateLocalFileReader(FullDemoFilename);
 			if (LocalFileAr.IsValid())
 			{
-				if (ReadReplayInfo(*LocalFileAr, RequestData.ReplayInfo, EReadReplayInfoFlags::None))
-				{
-					LocalFileAr->Seek(RequestData.ReplayInfo.Checkpoints[CheckpointIndex].EventDataOffset);
+				LocalFileAr->Seek(TaskReplayInfo.Checkpoints[CheckpointIndex].EventDataOffset);
 
-					RequestData.DataBuffer.AddUninitialized(RequestData.ReplayInfo.Checkpoints[CheckpointIndex].SizeInBytes);
+				RequestData.DataBuffer.AddUninitialized(TaskReplayInfo.Checkpoints[CheckpointIndex].SizeInBytes);
 
 					LocalFileAr->Serialize(RequestData.DataBuffer.GetData(), RequestData.DataBuffer.Num());
 
 					// Get the checkpoint data
 					
-					if (RequestData.ReplayInfo.bEncrypted)
+				if (TaskReplayInfo.bEncrypted)
 					{
 						if (SupportsEncryption())
 						{
@@ -2320,7 +2437,7 @@ void FLocalFileNetworkReplayStreamer::GotoCheckpointIndex(const int32 Checkpoint
 
 							TArray<uint8> DecryptedData;
 
-							if (!DecryptBuffer(RequestData.DataBuffer, DecryptedData, RequestData.ReplayInfo.EncryptionKey))
+						if (!DecryptBuffer(RequestData.DataBuffer, DecryptedData, TaskReplayInfo.EncryptionKey))
 							{
 								UE_LOG(LogLocalFileReplay, Error, TEXT("FLocalFileNetworkReplayStreamer::GotoCheckpointIndex. DecryptBuffer FAILED."));
 								RequestData.DataBuffer.Empty();
@@ -2337,7 +2454,7 @@ void FLocalFileNetworkReplayStreamer::GotoCheckpointIndex(const int32 Checkpoint
 						}
 					}
 
-					if (RequestData.ReplayInfo.bCompressed)
+				if (TaskReplayInfo.bCompressed)
 					{
 						if (SupportsCompression())
 						{
@@ -2361,7 +2478,6 @@ void FLocalFileNetworkReplayStreamer::GotoCheckpointIndex(const int32 Checkpoint
 							return;
 						}
 					}
-				}
 
 				LocalFileAr = nullptr;
 			}
@@ -2988,26 +3104,22 @@ void FLocalFileNetworkReplayStreamer::WriteHeader()
 		return;
 	}
 
-	int32 HeaderChunkIndex = CurrentReplayInfo.HeaderChunkIndex;
-	int64 HeaderTypeOffset = (HeaderChunkIndex != INDEX_NONE) ? CurrentReplayInfo.Chunks[HeaderChunkIndex].TypeOffset : 0;
-	int32 HeaderSize = (HeaderChunkIndex != INDEX_NONE) ? CurrentReplayInfo.Chunks[HeaderChunkIndex].SizeInBytes : 0;
-
-	AddGenericRequestToQueue<FLocalFileReplayInfo>(EQueuedLocalFileRequestType::WritingHeader, 
-		[this, HeaderChunkIndex, HeaderTypeOffset, HeaderSize, HeaderData=MoveTemp(HeaderAr.Buffer)](FLocalFileReplayInfo& ReplayInfo)
+	AddGenericRequestToQueue<ELocalFileReplayResult>(EQueuedLocalFileRequestType::WritingHeader,
+		[this, HeaderData=MoveTemp(HeaderAr.Buffer)](ELocalFileReplayResult& ReplayResult)
 		{
 			SCOPE_CYCLE_COUNTER(STAT_LocalReplay_FlushHeader);
 
 			TSharedPtr<FArchive> LocalFileAr = CreateLocalFileWriter(GetDemoFullFilename(CurrentStreamName));
 			if (LocalFileAr.IsValid())
 			{
-				if (HeaderChunkIndex == INDEX_NONE)
+				if (TaskReplayInfo.HeaderChunkIndex == INDEX_NONE)
 				{
 					// not expecting an existing header on disk, so check for it having been written by another process/client
 					FLocalFileReplayInfo TestInfo;
 					if (ReadReplayInfo(CurrentStreamName, TestInfo))
 					{
 						UE_LOG(LogLocalFileReplay, Error, TEXT("FLocalFileNetworkReplayStreamer::WriteHeader - Current file already has unexpected header"));
-						ReplayInfo.bIsValid = false;
+						ReplayResult = ELocalFileReplayResult::InvalidReplayInfo;
 						return;
 					}
 
@@ -3016,45 +3128,64 @@ void FLocalFileNetworkReplayStreamer::WriteHeader()
 				}
 				else 
 				{
-					if (HeaderData.Num() > HeaderSize)
-					{
-						LocalFileAr->Seek(HeaderTypeOffset);
+					const int32 HeaderSize = TaskReplayInfo.Chunks[TaskReplayInfo.HeaderChunkIndex].SizeInBytes;
+
+					LocalFileAr->Seek(TaskReplayInfo.Chunks[TaskReplayInfo.HeaderChunkIndex].TypeOffset);
 
+					// reuse existing chunk if the size hasn't changed, otherwise, add a new one
+					//@todo: add header data size, so we can differentiate and put a smaller header into the existing chunk space
+					if (HeaderData.Num() != HeaderSize)
+					{
 						// clear chunk type so it will be skipped later
 						ELocalFileChunkType ChunkType = ELocalFileChunkType::Unknown;
 						*LocalFileAr << ChunkType;
 
 						LocalFileAr->Seek(LocalFileAr->TotalSize());
-					}
-					else
-					{
-						LocalFileAr->Seek(HeaderTypeOffset);
+
+						TaskReplayInfo.Chunks[TaskReplayInfo.HeaderChunkIndex].ChunkType = ChunkType;
+						TaskReplayInfo.HeaderChunkIndex = INDEX_NONE;
 					}
 				}
 
+				const int32 TypeOffset = LocalFileAr->Tell();
+
 				ELocalFileChunkType ChunkType = ELocalFileChunkType::Header;
 				*LocalFileAr << ChunkType;
 
 				int32 ChunkSize = HeaderData.Num();
 				*LocalFileAr << ChunkSize;
 
-				LocalFileAr->Serialize((void*)HeaderData.GetData(), HeaderData.Num());
+				const int32 DataOffset = LocalFileAr->Tell();
 
+				LocalFileAr->Serialize((void*)HeaderData.GetData(), HeaderData.Num());
 				LocalFileAr = nullptr;
-			}
 
-			ReadReplayInfo(CurrentStreamName, ReplayInfo);
+				if (TaskReplayInfo.HeaderChunkIndex == INDEX_NONE)
+				{
+					FLocalFileChunkInfo& ChunkInfo = TaskReplayInfo.Chunks.AddDefaulted_GetRef();
+					ChunkInfo.ChunkType = ChunkType;
+					ChunkInfo.TypeOffset = TypeOffset;
+					ChunkInfo.DataOffset = DataOffset;
+					ChunkInfo.SizeInBytes = ChunkSize;
+
+					TaskReplayInfo.HeaderChunkIndex = TaskReplayInfo.Chunks.Num() - 1;
+				}
+			}
+			else
+			{
+				ReplayResult = ELocalFileReplayResult::FileWriter;
+			}
 		},
-		[this](FLocalFileReplayInfo& ReplayInfo)
+		[this](ELocalFileReplayResult& ReplayResult)
 		{
-			if (ReplayInfo.bIsValid)
+			if (ReplayResult == ELocalFileReplayResult::Success)
 			{
-				UpdateCurrentReplayInfo(ReplayInfo);
+				UpdateCurrentReplayInfo(TaskReplayInfo, EUpdateReplayInfoFlags::None);
 			}
 			else
 			{
-				UE_LOG(LogLocalFileReplay, Error, TEXT("FLocalFileNetworkReplayStreamer::WriteHeader failed with invalid replay info."));
-				SetLastError(ELocalFileReplayResult::InvalidReplayInfo);
+				UE_LOG(LogLocalFileReplay, Error, TEXT("FLocalFileNetworkReplayStreamer::WriteHeader failed."));
+				SetLastError(ReplayResult);
 			}
 		});
 
@@ -3194,9 +3325,7 @@ void FLocalFileNetworkReplayStreamer::ConditionallyLoadNextChunk()
 			SCOPE_CYCLE_COUNTER(STAT_LocalReplay_ReadStream);
 			LLM_SCOPE(ELLMTag::Replays);
 
-			if (ReadReplayInfo(CurrentStreamName, RequestData.ReplayInfo))
-			{
-				check(RequestData.ReplayInfo.DataChunks.IsValidIndex(RequestedStreamChunkIndex));
+			check(TaskReplayInfo.DataChunks.IsValidIndex(RequestedStreamChunkIndex));
 
 				RequestData.DataBuffer.Empty();
 			
@@ -3205,20 +3334,20 @@ void FLocalFileNetworkReplayStreamer::ConditionallyLoadNextChunk()
 				TSharedPtr<FArchive> LocalFileAr = CreateLocalFileReader(FullDemoFilename);
 				if (LocalFileAr.IsValid())
 				{
-					LocalFileAr->Seek(RequestData.ReplayInfo.DataChunks[RequestedStreamChunkIndex].ReplayDataOffset);
+				LocalFileAr->Seek(TaskReplayInfo.DataChunks[RequestedStreamChunkIndex].ReplayDataOffset);
 
-					RequestData.DataBuffer.AddUninitialized(RequestData.ReplayInfo.DataChunks[RequestedStreamChunkIndex].SizeInBytes);
+				RequestData.DataBuffer.AddUninitialized(TaskReplayInfo.DataChunks[RequestedStreamChunkIndex].SizeInBytes);
 
 					LocalFileAr->Serialize(RequestData.DataBuffer.GetData(), RequestData.DataBuffer.Num());
 
-					if (RequestData.ReplayInfo.bEncrypted)
+				if (TaskReplayInfo.bEncrypted)
 					{
 						if (SupportsEncryption())
 						{
 							SCOPE_CYCLE_COUNTER(STAT_LocalReplay_DecryptTime);
 
 							TArray<uint8> DecryptedData;
-							if (DecryptBuffer(RequestData.DataBuffer, DecryptedData, RequestData.ReplayInfo.EncryptionKey))
+						if (DecryptBuffer(RequestData.DataBuffer, DecryptedData, TaskReplayInfo.EncryptionKey))
 							{
 								RequestData.DataBuffer = MoveTemp(DecryptedData);
 							}
@@ -3239,7 +3368,7 @@ void FLocalFileNetworkReplayStreamer::ConditionallyLoadNextChunk()
 						}
 					}
 
-					if (RequestData.ReplayInfo.bCompressed)
+				if (TaskReplayInfo.bCompressed)
 					{
 						if (SupportsCompression())
 						{
@@ -3269,7 +3398,6 @@ void FLocalFileNetworkReplayStreamer::ConditionallyLoadNextChunk()
 
 					LocalFileAr = nullptr;
 				}
-			}
 		},
 		[this, RequestedStreamChunkIndex](TLocalFileRequestCommonData<FStreamingResultBase>& RequestData)
 		{
@@ -3359,19 +3487,19 @@ void FLocalFileNetworkReplayStreamer::ConditionallyRefreshReplayInfo()
 
 		if (FPlatformTime::Seconds() - LastRefreshTime > REFRESH_REPLAYINFO_IN_SECONDS)
 		{
-			int64 LastDataSize = CurrentReplayInfo.TotalDataSizeInBytes;
+			const int64 LastDataSize = CurrentReplayInfo.TotalDataSizeInBytes;
 			const FString FullDemoFilename = GetDemoFullFilename(CurrentStreamName);
 
-			AddGenericRequestToQueue<FLocalFileReplayInfo>(EQueuedLocalFileRequestType::RefreshingLiveStream, 
-				[this](FLocalFileReplayInfo& ReplayInfo)
+			AddGenericRequestToQueue<ELocalFileReplayResult>(EQueuedLocalFileRequestType::RefreshingLiveStream, 
+				[this](ELocalFileReplayResult& ReplayResult)
 				{
-					ReadReplayInfo(CurrentStreamName, ReplayInfo);
+					ReadReplayInfo(CurrentStreamName, TaskReplayInfo);
 				},
-				[this, LastDataSize](FLocalFileReplayInfo& ReplayInfo)
+				[this, LastDataSize](ELocalFileReplayResult& ReplayResult)
 				{
-					if (ReplayInfo.bIsValid && (ReplayInfo.TotalDataSizeInBytes != LastDataSize))
+					if (TaskReplayInfo.bIsValid && (TaskReplayInfo.TotalDataSizeInBytes != LastDataSize))
 					{
-						CurrentReplayInfo = ReplayInfo;
+						UpdateCurrentReplayInfo(TaskReplayInfo, EUpdateReplayInfoFlags::FullUpdate);
 					}
 				});
 
@@ -3505,7 +3633,7 @@ bool FLocalFileNetworkReplayStreamer::IsCheckpointTypeSupported(EReplayCheckpoin
 	return bSupported;
 }
 
-void FLocalFileNetworkReplayStreamer::UpdateCurrentReplayInfo(FLocalFileReplayInfo& ReplayInfo)
+void FLocalFileNetworkReplayStreamer::UpdateCurrentReplayInfo(FLocalFileReplayInfo& ReplayInfo, EUpdateReplayInfoFlags UpdateFlags)
 {
 	if (ensure(ReplayInfo.bIsValid))
 	{
@@ -3515,10 +3643,13 @@ void FLocalFileNetworkReplayStreamer::UpdateCurrentReplayInfo(FLocalFileReplayIn
 
 		CurrentReplayInfo = ReplayInfo;
 		
+		if (!EnumHasAnyFlags(UpdateFlags, EUpdateReplayInfoFlags::FullUpdate))
+		{
 		CurrentReplayInfo.LengthInMS = TotalLengthInMS;
 		CurrentReplayInfo.EncryptionKey = MoveTemp(CurrentKey);
 	}
 }
+}
 
 int32 FLocalFileNetworkReplayStreamer::GetDecompressedSizeBackCompat(FArchive& InCompressed) const
 {
diff --git a/Engine/Source/Runtime/NetworkReplayStreaming/LocalFileNetworkReplayStreaming/Public/LocalFileNetworkReplayStreaming.h b/Engine/Source/Runtime/NetworkReplayStreaming/LocalFileNetworkReplayStreaming/Public/LocalFileNetworkReplayStreaming.h
--- a/Engine/Source/Runtime/NetworkReplayStreaming/LocalFileNetworkReplayStreaming/Public/LocalFileNetworkReplayStreaming.h
+++ b/Engine/Source/Runtime/NetworkReplayStreaming/LocalFileNetworkReplayStreaming/Public/LocalFileNetworkReplayStreaming.h
@@ -61,6 +61,14 @@ enum class EReadReplayInfoFlags : uint32
 
 ENUM_CLASS_FLAGS(EReadReplayInfoFlags);
 
+enum class EUpdateReplayInfoFlags : uint32
+{
+	None = 0,
+	FullUpdate = 1,
+};
+
+ENUM_CLASS_FLAGS(EUpdateReplayInfoFlags);
+
 /** Struct to hold chunk metadata */
 struct FLocalFileChunkInfo
 {
@@ -120,6 +128,8 @@ struct FLocalFileEventInfo
 
 	int32 SizeInBytes;
 	int64 EventDataOffset;
+
+	void CountBytes(FArchive& Ar) const;
 };
 
 /** Struct to hold metadata about an entire replay */
@@ -157,6 +167,8 @@ struct FLocalFileReplayInfo
 	TArray<FLocalFileEventInfo> Checkpoints;
 	TArray<FLocalFileEventInfo> Events;
 	TArray<FLocalFileReplayDataInfo> DataChunks;
+
+	void CountBytes(FArchive& Ar) const;
 };
 
 /** Archive to wrap the file reader and respect chunk boundaries */
@@ -269,6 +281,10 @@ enum class ELocalFileReplayResult : uint32
 	CompressionNotSupported,
 	DecryptBuffer,
 	EncryptionNotSupported,
+	EncryptBuffer,
+	CompressBuffer,
+	InvalidName,
+	FileWriter,
 	Unknown,
 };
 
@@ -381,6 +397,7 @@ class TGenericQueuedLocalFileRequest : public FQueuedLocalFileRequest, public TS
 public:
 	TGenericQueuedLocalFileRequest(const TSharedPtr<FLocalFileNetworkReplayStreamer>& InStreamer, EQueuedLocalFileRequestType::Type InType, TFunction<void(StorageType&)>&& InFunction, TFunction<void(StorageType&)>&& InCompletionCallback)
 		: FQueuedLocalFileRequest(InStreamer, InType)
+		, Storage()
 		, RequestFunction(MoveTemp(InFunction))
 		, CompletionCallback(MoveTemp(InCompletionCallback))
 	{
@@ -432,7 +449,10 @@ class TLocalFileRequestCommonData
 {
 public:
 	DelegateResultType DelegateResult;
+	
+	UE_DEPRECATED(5.3, "No longer used")
 	FLocalFileReplayInfo ReplayInfo;
+	
 	TArray<uint8> DataBuffer;
 
 	UE_DEPRECATED(5.1, "No longer used")
@@ -615,7 +635,13 @@ protected:
 	void RenameReplayFriendlyName_Internal(const FString& ReplayName, const FString& NewFriendlyName, const int32 UserIndex, const FRenameReplayCallback& Delegate);
 	void RenameReplay_Internal(const FString& ReplayName, const FString& NewName, const int32 UserIndex, const FRenameReplayCallback& Delegate);
 
-	/** Currently playing or recording replay metadata */
+	/** 
+	 * Currently playing or recording replay metadata
+	 * 
+	 * The values may not accurately reflect what is currently on disk during recording.  They will be updated 
+	 * as each queued task completes on the game thread.
+	 * 
+	 **/
 	FLocalFileReplayInfo CurrentReplayInfo;
 
 	TInterval<uint32> StreamTimeRange;
@@ -725,7 +751,7 @@ protected:
 	mutable TMap<FString, TArray<uint8>> FileContentsCache;
 	const TArray<uint8>& GetCachedFileContents(const FString& Filename) const;
 
-	void UpdateCurrentReplayInfo(FLocalFileReplayInfo& ReplayInfo);
+	void UpdateCurrentReplayInfo(FLocalFileReplayInfo& ReplayInfo, EUpdateReplayInfoFlags UpdateFlags = EUpdateReplayInfoFlags::None);
 
 	virtual int32 GetDecompressedSizeBackCompat(FArchive& InCompressed) const;
 
@@ -740,6 +766,10 @@ public:
 
 	UE_DEPRECATED(5.2, "No longer used, replaced with custom version.")
 	static const uint32 LatestVersion;
+
+private:
+	/** Manipulated by queued replay tasks, likely from another thread */
+	FLocalFileReplayInfo TaskReplayInfo;
 };
 
 class LOCALFILENETWORKREPLAYSTREAMING_API FLocalFileNetworkReplayStreamingFactory : public INetworkReplayStreamingFactory, public FTickableGameObject
