Lower verbosity of spammy replay read log line

diff --git a/Engine/Source/Runtime/NetworkReplayStreaming/LocalFileNetworkReplayStreaming/Private/LocalFileNetworkReplayStreaming.cpp b/Engine/Source/Runtime/NetworkReplayStreaming/LocalFileNetworkReplayStreaming/Private/LocalFileNetworkReplayStreaming.cpp
--- a/Engine/Source/Runtime/NetworkReplayStreaming/LocalFileNetworkReplayStreaming/Private/LocalFileNetworkReplayStreaming.cpp
+++ b/Engine/Source/Runtime/NetworkReplayStreaming/LocalFileNetworkReplayStreaming/Private/LocalFileNetworkReplayStreaming.cpp
@@ -577,7 +577,7 @@ bool FLocalFileNetworkReplayStreamer::ReadReplayInfo(FArchive& Archive, FLocalFi
 		Info.CountBytes(CountMemAr);
 		const int64 InfoSize = sizeof(FLocalFileReplayInfo) + CountMemAr.GetNum();
 
-		UE_LOG(LogLocalFileReplay, Log, TEXT("ReadReplayInfo: IsValid: %s MemSize: %lld bytes"), *LexToString(Info.bIsValid), InfoSize);
+		UE_LOG(LogLocalFileReplay, Verbose, TEXT("ReadReplayInfo: IsValid: %s MemSize: %lld bytes"), *LexToString(Info.bIsValid), InfoSize);
 
 		return Info.bIsValid && !Archive.IsError();
 	}
