Enabling UnsafeTypeCastWarningLevel for replay streamers

diff --git a/Engine/Source/Runtime/NetworkReplayStreaming/HttpNetworkReplayStreaming/HttpNetworkReplayStreaming.Build.cs b/Engine/Source/Runtime/NetworkReplayStreaming/HttpNetworkReplayStreaming/HttpNetworkReplayStreaming.Build.cs
--- a/Engine/Source/Runtime/NetworkReplayStreaming/HttpNetworkReplayStreaming/HttpNetworkReplayStreaming.Build.cs
+++ b/Engine/Source/Runtime/NetworkReplayStreaming/HttpNetworkReplayStreaming/HttpNetworkReplayStreaming.Build.cs
@@ -24,6 +24,8 @@ namespace UnrealBuildTool.Rules
                 AddEngineThirdPartyPrivateStaticDependencies(Target,"libWebSockets");
                 AddEngineThirdPartyPrivateStaticDependencies(Target,"zlib");
             }
-        }
+
+			UnsafeTypeCastWarningLevel = WarningLevel.Error;
+		}
     }
 }
diff --git a/Engine/Source/Runtime/NetworkReplayStreaming/HttpNetworkReplayStreaming/Private/HttpNetworkReplayStreaming.cpp b/Engine/Source/Runtime/NetworkReplayStreaming/HttpNetworkReplayStreaming/Private/HttpNetworkReplayStreaming.cpp
--- a/Engine/Source/Runtime/NetworkReplayStreaming/HttpNetworkReplayStreaming/Private/HttpNetworkReplayStreaming.cpp
+++ b/Engine/Source/Runtime/NetworkReplayStreaming/HttpNetworkReplayStreaming/Private/HttpNetworkReplayStreaming.cpp
@@ -154,32 +154,33 @@ public:
 
 void FHttpStreamFArchive::Serialize( void* V, int64 Length ) 
 {
-	if ( IsLoading() )
+	if (IsLoading())
 	{
-		if ( Pos + Length > Buffer.Num() )
+		if ((Pos + Length) > Buffer.Num())
 		{
 			SetError();
 			return;
 		}
 			
-		FMemory::Memcpy( V, Buffer.GetData() + Pos, Length );
+		FMemory::Memcpy(V, Buffer.GetData() + Pos, Length);
 
-		Pos += Length;
+		Pos += IntCastChecked<int32>(Length);
 	}
 	else
 	{
-		check( Pos <= Buffer.Num() );
+		check(Pos <= Buffer.Num());
 
-		const int32 SpaceNeeded = Length - ( Buffer.Num() - Pos );
+		const int32 LocalLength = IntCastChecked<int32>(Length);
+		const int32 SpaceNeeded = LocalLength - (Buffer.Num() - Pos);
 
-		if ( SpaceNeeded > 0 )
+		if (SpaceNeeded > 0)
 		{
-			Buffer.AddZeroed( SpaceNeeded );
+			Buffer.AddZeroed(SpaceNeeded);
 		}
 
-		FMemory::Memcpy( Buffer.GetData() + Pos, V, Length );
+		FMemory::Memcpy(Buffer.GetData() + Pos, V, LocalLength);
 
-		Pos += Length;
+		Pos += LocalLength;
 	}
 }
 
@@ -197,7 +198,7 @@ void FHttpStreamFArchive::Seek( int64 InPos )
 {
 	check( InPos <= Buffer.Num() );
 
-	Pos = InPos;
+	Pos = IntCastChecked<int32>(InPos);
 }
 
 bool FHttpStreamFArchive::AtEnd() 
@@ -749,7 +750,7 @@ bool FQueuedGotoFakeCheckpoint::PreProcess( FHttpNetworkReplayStreamer* Streamer
 	Streamer->CheckpointArchive.Buffer.Empty();
 	Streamer->CheckpointArchive.Pos = 0;
 
-	if (!Streamer->IsDataAvailableForTimeRange(0, Streamer->LastGotoTimeInMS))
+	if (!Streamer->IsDataAvailableForTimeRange(0, IntCastChecked<uint32>(Streamer->LastGotoTimeInMS)))
 	{
 		// Completely reset our stream (we're going to start downloading from the start of the checkpoint)
 		Streamer->StreamArchive.Buffer.Empty();
@@ -810,7 +811,7 @@ void FHttpNetworkReplayStreamer::GotoCheckpointIndex( const int32 CheckpointInde
 	if ( CheckpointIndex == -1 )
 	{
 		GotoCheckpointDelegate = Delegate;
-		SetHighPriorityTimeRange( 0, LastGotoTimeInMS );
+		SetHighPriorityTimeRange(0, IntCastChecked<uint32>(LastGotoTimeInMS));
 		LastChunkTime = 0;		// Force the next chunk to start downloading immediately
 		AddCustomRequestToQueue( TSharedPtr< FQueuedHttpRequest >( new FQueuedGotoFakeCheckpoint() ) );
 		return;
@@ -2184,7 +2185,7 @@ void FHttpNetworkReplayStreamer::HttpDownloadCheckpointFinished( FHttpRequestPtr
 		if ( LastGotoTimeInMS >= 0 )
 		{
 			// If we are fine scrubbing, make sure to wait on the part of the stream that is needed to do this in one frame
-			SetHighPriorityTimeRange( Checkpoint.Time1, LastGotoTimeInMS );
+			SetHighPriorityTimeRange(Checkpoint.Time1, IntCastChecked<uint32>(LastGotoTimeInMS));
 			
 			// Subtract off checkpoint time so we pass in the leftover to the engine to fast forward through for the fine scrubbing part
 			LastGotoTimeInMS -= Checkpoint.Time1;
@@ -2653,7 +2654,7 @@ void FHttpNetworkReplayStreamer::HttpDownloadCheckpointDeltaFinished( FHttpReque
 			if ( LastGotoTimeInMS >= 0 )
 			{
 				// If we are fine scrubbing, make sure to wait on the part of the stream that is needed to do this in one frame
-				SetHighPriorityTimeRange( CheckpointList.ReplayEvents[ DownloadCheckpointIndex ].Time1, LastGotoTimeInMS );
+				SetHighPriorityTimeRange(CheckpointList.ReplayEvents[ DownloadCheckpointIndex ].Time1, IntCastChecked<uint32>(LastGotoTimeInMS));
 
 				// Subtract off checkpoint time so we pass in the leftover to the engine to fast forward through for the fine scrubbing part
 				LastGotoTimeInMS -= CheckpointList.ReplayEvents[ DownloadCheckpointIndex ].Time1;
diff --git a/Engine/Source/Runtime/NetworkReplayStreaming/InMemoryNetworkReplayStreaming/InMemoryNetworkReplayStreaming.build.cs b/Engine/Source/Runtime/NetworkReplayStreaming/InMemoryNetworkReplayStreaming/InMemoryNetworkReplayStreaming.build.cs
--- a/Engine/Source/Runtime/NetworkReplayStreaming/InMemoryNetworkReplayStreaming/InMemoryNetworkReplayStreaming.build.cs
+++ b/Engine/Source/Runtime/NetworkReplayStreaming/InMemoryNetworkReplayStreaming/InMemoryNetworkReplayStreaming.build.cs
@@ -23,6 +23,8 @@ namespace UnrealBuildTool.Rules
                     "Json"
 				}
 			);
+
+			UnsafeTypeCastWarningLevel = WarningLevel.Error;
 		}
 	}
 }
diff --git a/Engine/Source/Runtime/NetworkReplayStreaming/InMemoryNetworkReplayStreaming/Private/InMemoryNetworkReplayStreaming.cpp b/Engine/Source/Runtime/NetworkReplayStreaming/InMemoryNetworkReplayStreaming/Private/InMemoryNetworkReplayStreaming.cpp
--- a/Engine/Source/Runtime/NetworkReplayStreaming/InMemoryNetworkReplayStreaming/Private/InMemoryNetworkReplayStreaming.cpp
+++ b/Engine/Source/Runtime/NetworkReplayStreaming/InMemoryNetworkReplayStreaming/Private/InMemoryNetworkReplayStreaming.cpp
@@ -347,7 +347,7 @@ FArchive* FInMemoryNetworkReplayStreamer::GetCheckpointArchive()
 		if (TimeBufferHintSeconds > 0.0f)
 		{
 			// Absolute time at which the buffer should start
-			const float BufferStartTimeMS = FoundReplay->StreamInfo.LengthInMS - (TimeBufferHintSeconds * 1000.0);
+			const uint32 BufferStartTimeMS = UE::LWC::FloatToIntCastChecked<uint32>(FoundReplay->StreamInfo.LengthInMS - (TimeBufferHintSeconds * 1000.0));
 
 			// Store the found checkpoint's time so that we can line up chunks with it
 			uint32 FoundCheckpointTime = 0;
@@ -409,7 +409,7 @@ void FInMemoryNetworkReplayStreamer::FlushCheckpoint(const uint32 TimeInMS)
 	CheckpointAr.Reset();
 	
 	CheckpointCurrentlyBeingSaved.TimeInMS = TimeInMS;
-	CheckpointCurrentlyBeingSaved.StreamByteOffset = FileAr->Tell();
+	CheckpointCurrentlyBeingSaved.StreamByteOffset = IntCastChecked<uint32>(FileAr->Tell());
 
 	FInMemoryReplay* const FoundReplay = GetCurrentReplayChecked();
 
@@ -578,7 +578,7 @@ void FInMemoryReplayStreamArchive::Serialize(void* V, int64 Length)
 			return;
 		}
 
-		const int32 OffsetIntoChunk = Pos - CurrentChunk->StartIndex;
+		const int64 OffsetIntoChunk = Pos - CurrentChunk->StartIndex;
 		FMemory::Memcpy(V, CurrentChunk->Data.GetData() + OffsetIntoChunk, Length);
 
 		Pos += Length;
@@ -595,9 +595,8 @@ void FInMemoryReplayStreamArchive::Serialize(void* V, int64 Length)
 			return;
 		}
 
-		const int32 OffsetIntoChunk = Pos - CurrentChunk->StartIndex;
-		const int32 SpaceNeeded = Length - (CurrentChunk->Data.Num() - OffsetIntoChunk);
-
+		const int64 OffsetIntoChunk = Pos - CurrentChunk->StartIndex;
+		const int32 SpaceNeeded = IntCastChecked<int32>(Length - (CurrentChunk->Data.Num() - OffsetIntoChunk));
 		if (SpaceNeeded > 0)
 		{
 			CurrentChunk->Data.AddZeroed(SpaceNeeded);
diff --git a/Engine/Source/Runtime/NetworkReplayStreaming/InMemoryNetworkReplayStreaming/Public/InMemoryNetworkReplayStreaming.h b/Engine/Source/Runtime/NetworkReplayStreaming/InMemoryNetworkReplayStreaming/Public/InMemoryNetworkReplayStreaming.h
--- a/Engine/Source/Runtime/NetworkReplayStreaming/InMemoryNetworkReplayStreaming/Public/InMemoryNetworkReplayStreaming.h
+++ b/Engine/Source/Runtime/NetworkReplayStreaming/InMemoryNetworkReplayStreaming/Public/InMemoryNetworkReplayStreaming.h
@@ -97,16 +97,16 @@ public:
 		, Chunks(InChunks)
 	{}
 
-	virtual void	Serialize( void* V, int64 Length );
+	virtual void	Serialize(void* V, int64 Length);
 	virtual int64	Tell();
 	virtual int64	TotalSize();
-	virtual void	Seek( int64 InPos );
+	virtual void	Seek(int64 InPos);
 	virtual bool	AtEnd();
 
 private:
 	FInMemoryReplay::FStreamChunk* GetCurrentChunk() const;
 
-	int32 Pos;
+	int64 Pos;
 	TArray<FInMemoryReplay::FStreamChunk>& Chunks;
 };
 
diff --git a/Engine/Source/Runtime/NetworkReplayStreaming/LocalFileNetworkReplayStreaming/LocalFileNetworkReplayStreaming.Build.cs b/Engine/Source/Runtime/NetworkReplayStreaming/LocalFileNetworkReplayStreaming/LocalFileNetworkReplayStreaming.Build.cs
--- a/Engine/Source/Runtime/NetworkReplayStreaming/LocalFileNetworkReplayStreaming/LocalFileNetworkReplayStreaming.Build.cs
+++ b/Engine/Source/Runtime/NetworkReplayStreaming/LocalFileNetworkReplayStreaming/LocalFileNetworkReplayStreaming.Build.cs
@@ -16,6 +16,8 @@ namespace UnrealBuildTool.Rules
 					"NetworkReplayStreaming",
 					"Json"
 				} );
-        }
-    }
+
+			UnsafeTypeCastWarningLevel = WarningLevel.Error;
+		}
+	}
 }
diff --git a/Engine/Source/Runtime/NetworkReplayStreaming/LocalFileNetworkReplayStreaming/Private/LocalFileNetworkReplayStreaming.cpp b/Engine/Source/Runtime/NetworkReplayStreaming/LocalFileNetworkReplayStreaming/Private/LocalFileNetworkReplayStreaming.cpp
--- a/Engine/Source/Runtime/NetworkReplayStreaming/LocalFileNetworkReplayStreaming/Private/LocalFileNetworkReplayStreaming.cpp
+++ b/Engine/Source/Runtime/NetworkReplayStreaming/LocalFileNetworkReplayStreaming/Private/LocalFileNetworkReplayStreaming.cpp
@@ -123,37 +123,51 @@ void FLocalFileStreamFArchive::Serialize(void* V, int64 Length)
 {
 	if (IsLoading())
 	{
-		if ((Length < 0) || (Pos + Length) > Buffer.Num())
+		if ((Length < 0) || (ArchivePos + Length) > Buffer.Num())
 		{
-			UE_LOG(LogLocalFileReplay, Error, TEXT("FLocalFileStreamFArchive::Serialize: Attempted to serialize past end of archive: Position = %i, Size=%i, Requested = %lli"), Pos, Buffer.Num(), Length);
+			UE_LOG(LogLocalFileReplay, Error, TEXT("FLocalFileStreamFArchive::Serialize: Attempted to serialize past end of archive: Position = %i, Size=%i, Requested = %lli"), ArchivePos, Buffer.Num(), Length);
 			SetError();
 			return;
 		}
 
-		FMemory::Memcpy(V, Buffer.GetData() + Pos, Length);
+		FMemory::Memcpy(V, Buffer.GetData() + ArchivePos, Length);
 
-		Pos += Length;
+		ArchivePos += Length;
 	}
 	else
 	{
-		check(Pos <= Buffer.Num());
-
-		const int32 SpaceNeeded = Length - (Buffer.Num() - Pos);
+		check(ArchivePos <= Buffer.Num());
 
+		const int32 SpaceNeeded = IntCastChecked<int32>(Length - (Buffer.Num() - ArchivePos));
 		if (SpaceNeeded > 0)
 		{
 			Buffer.AddUninitialized(SpaceNeeded);
 		}
 
-		FMemory::Memcpy(Buffer.GetData() + Pos, V, Length);
+		FMemory::Memcpy(Buffer.GetData() + ArchivePos, V, Length);
 
-		Pos += Length;
+		ArchivePos += Length;
 	}
+
+	PRAGMA_DISABLE_DEPRECATION_WARNINGS
+	Pos = static_cast<int32>(ArchivePos);
+	PRAGMA_ENABLE_DEPRECATION_WARNINGS
 }
 
-int64 FLocalFileStreamFArchive::Tell() 
+int64 FLocalFileStreamFArchive::Tell() const
 {
-	return Pos;
+	return ArchivePos;
+}
+
+int64 FLocalFileStreamFArchive::Tell()
+{
+	return ArchivePos;
+}
+
+
+int64 FLocalFileStreamFArchive::TotalSize() const
+{
+	return Buffer.Num();
 }
 
 int64 FLocalFileStreamFArchive::TotalSize()
@@ -165,12 +179,12 @@ void FLocalFileStreamFArchive::Seek(int64 InPos)
 {
 	check(InPos <= Buffer.Num());
 
-	Pos = InPos;
+	ArchivePos = InPos;
 }
 
 bool FLocalFileStreamFArchive::AtEnd() 
 {
-	return Pos >= Buffer.Num() && bAtEndOfReplay;
+	return ArchivePos >= Buffer.Num() && bAtEndOfReplay;
 }
 
 FLocalFileNetworkReplayStreamer::FLocalFileNetworkReplayStreamer() 
@@ -1048,7 +1062,7 @@ bool FLocalFileNetworkReplayStreamer::IsDataAvailable() const
 	}
 
 	// If we are loading, and we have more data
-	if (StreamAr.IsLoading() && StreamAr.Pos < StreamAr.Buffer.Num() && CurrentReplayInfo.DataChunks.Num() > 0)
+	if (StreamAr.IsLoading() && StreamAr.Tell() < StreamAr.TotalSize() && CurrentReplayInfo.DataChunks.Num() > 0)
 	{
 		return true;
 	}
@@ -1214,49 +1228,49 @@ void FLocalFileNetworkReplayStreamer::AddOrUpdateEvent(const FString& Name, cons
 			// prefix with stream name to be consistent with http streamer
 			EventName = CurrentStreamName + TEXT("_") + EventName;
 
-				TSharedPtr<FArchive> LocalFileAr = CreateLocalFileWriter(GetDemoFullFilename(CurrentStreamName));
-				if (LocalFileAr.IsValid())
-				{
-					int32 EventIndex = INDEX_NONE;
+			TSharedPtr<FArchive> LocalFileAr = CreateLocalFileWriter(GetDemoFullFilename(CurrentStreamName));
+			if (LocalFileAr.IsValid())
+			{
+				int32 EventIndex = INDEX_NONE;
 
-					// see if this event already exists
+				// see if this event already exists
 				for (int32 i=0; i < TaskReplayInfo.Events.Num(); ++i)
-					{
+				{
 					if (TaskReplayInfo.Events[i].Id == EventName)
-						{
-							EventIndex = i;
-							break;
-						}
+					{
+						EventIndex = i;
+						break;
 					}
+				}
 
-					TArray<uint8> EncryptedData;
+				TArray<uint8> EncryptedData;
 
-					if (AllowEncryptedWrite())
-					{
-						SCOPE_CYCLE_COUNTER(STAT_LocalReplay_EncryptTime);
+				if (AllowEncryptedWrite())
+				{
+					SCOPE_CYCLE_COUNTER(STAT_LocalReplay_EncryptTime);
 
 					if (!EncryptBuffer(Data, EncryptedData, TaskReplayInfo.EncryptionKey))
-						{
-							UE_LOG(LogLocalFileReplay, Warning, TEXT("FLocalFileNetworkReplayStreamer::AddOrUpdateEvent - EncryptBuffer failed"));
-						ReplayResult = ELocalFileReplayResult::EncryptBuffer;
-							return;
-						}
-					}
-					else
 					{
-						EncryptedData = Data;
+						UE_LOG(LogLocalFileReplay, Warning, TEXT("FLocalFileNetworkReplayStreamer::AddOrUpdateEvent - EncryptBuffer failed"));
+						ReplayResult = ELocalFileReplayResult::EncryptBuffer;
+						return;
 					}
+				}
+				else
+				{
+					EncryptedData = Data;
+				}
 
-					// serialize event to temporary location
-					FArrayWriter Writer;
+				// serialize event to temporary location
+				FArrayWriter Writer;
 
-					ELocalFileChunkType ChunkType = ELocalFileChunkType::Event;
-					Writer << ChunkType;
+				ELocalFileChunkType ChunkType = ELocalFileChunkType::Event;
+				Writer << ChunkType;
 
 				const int64 SavedPos = Writer.Tell();
 
-					int32 PlaceholderSize = 0;
-					Writer << PlaceholderSize;
+				int32 PlaceholderSize = 0;
+				Writer << PlaceholderSize;
 
 				const int64 MetadataPos = Writer.Tell();
 
@@ -1269,44 +1283,44 @@ void FLocalFileNetworkReplayStreamer::AddOrUpdateEvent(const FString& Name, cons
 				FString MetaValue = Meta;
 				Writer << MetaValue;
 
-					uint32 Time1 = TimeInMS;
-					Writer << Time1;
+				uint32 Time1 = TimeInMS;
+				Writer << Time1;
 
-					uint32 Time2 = TimeInMS;
-					Writer << Time2;
+				uint32 Time2 = TimeInMS;
+				Writer << Time2;
 
-					int32 EventSize = EncryptedData.Num();
-					Writer << EventSize;
+				int32 EventSize = EncryptedData.Num();
+				Writer << EventSize;
 
 				const int64 InternalDataOffset = Writer.Tell();
 
-					Writer.Serialize((void*)EncryptedData.GetData(), EncryptedData.Num());
+				Writer.Serialize((void*)EncryptedData.GetData(), EncryptedData.Num());
 
-					int32 ChunkSize = Writer.Tell() - MetadataPos;
+				int32 ChunkSize = IntCastChecked<int32>(Writer.Tell() - MetadataPos);
 
 				bool bNewChunk = true;
 
-					if (EventIndex == INDEX_NONE)
-					{
-						// append new event chunk
-						LocalFileAr->Seek(LocalFileAr->TotalSize());
-					}
-					else 
-					{
+				if (EventIndex == INDEX_NONE)
+				{
+					// append new event chunk
+					LocalFileAr->Seek(LocalFileAr->TotalSize());
+				}
+				else 
+				{
 					if (ChunkSize > TaskReplayInfo.Chunks[TaskReplayInfo.Events[EventIndex].ChunkIndex].SizeInBytes)
-						{
+					{
 						LocalFileAr->Seek(TaskReplayInfo.Chunks[TaskReplayInfo.Events[EventIndex].ChunkIndex].TypeOffset);
 
-							// clear chunk type so it will be skipped later
-							ChunkType = ELocalFileChunkType::Unknown;
-							*LocalFileAr << ChunkType;
+						// clear chunk type so it will be skipped later
+						ChunkType = ELocalFileChunkType::Unknown;
+						*LocalFileAr << ChunkType;
 
 						TaskReplayInfo.Chunks[TaskReplayInfo.Events[EventIndex].ChunkIndex].ChunkType = ELocalFileChunkType::Unknown;
 
-							LocalFileAr->Seek(LocalFileAr->TotalSize());
-						}
-						else
-						{
+						LocalFileAr->Seek(LocalFileAr->TotalSize());
+					}
+					else
+					{
 						bNewChunk = false;
 
 						// reuse existing chunk
@@ -1315,7 +1329,7 @@ void FLocalFileNetworkReplayStreamer::AddOrUpdateEvent(const FString& Name, cons
 						// maintain the original chunk size to avoid corrupting the file
 						ChunkSize = TaskReplayInfo.Chunks[TaskReplayInfo.Events[EventIndex].ChunkIndex].SizeInBytes;
 					}
-						}
+				}
 
 				Writer.Seek(SavedPos);
 				Writer << ChunkSize;
@@ -1335,7 +1349,7 @@ void FLocalFileNetworkReplayStreamer::AddOrUpdateEvent(const FString& Name, cons
 					NewEventInfo.Id = MoveTemp(TempId);
 
 					EventIndex = TaskReplayInfo.Events.Num() - 1;
-					}
+				}
 
 				FLocalFileEventInfo& EventInfo = TaskReplayInfo.Events[EventIndex];
 				EventInfo.Group = MoveTemp(GroupValue);
@@ -1345,9 +1359,9 @@ void FLocalFileNetworkReplayStreamer::AddOrUpdateEvent(const FString& Name, cons
 				EventInfo.SizeInBytes = EventSize;
 				EventInfo.EventDataOffset = TypeOffset + InternalDataOffset;
 
-					LocalFileAr->Serialize(Writer.GetData(), Writer.TotalSize());
-					LocalFileAr = nullptr;
-				}
+				LocalFileAr->Serialize(Writer.GetData(), Writer.TotalSize());
+				LocalFileAr = nullptr;
+			}
 			else
 			{
 				ReplayResult = ELocalFileReplayResult::FileWriter;
@@ -1817,7 +1831,7 @@ void FLocalFileNetworkReplayStreamer::FlushStream(const uint32 TimeInMS)
 		return;
 	}
 
-	if (StreamAr.Buffer.Num() == 0)
+	if (StreamAr.TotalSize() == 0)
 	{
 		// Nothing to flush
 		return;
@@ -1833,7 +1847,7 @@ void FLocalFileNetworkReplayStreamer::FlushStream(const uint32 TimeInMS)
 	const int32 TotalLengthInMS = CurrentReplayInfo.LengthInMS;
 
 	// Save any newly streamed data to disk
-	UE_LOG(LogLocalFileReplay, Verbose, TEXT("FLocalFileNetworkReplayStreamer::FlushStream. StreamChunkIndex: %i, Size: %i"), StreamChunkIndex, StreamAr.Buffer.Num());
+	UE_LOG(LogLocalFileReplay, Verbose, TEXT("FLocalFileNetworkReplayStreamer::FlushStream. StreamChunkIndex: %i, Size: %" INT64_FMT), StreamChunkIndex, StreamAr.TotalSize());
 
 	AddGenericRequestToQueue<ELocalFileReplayResult>(EQueuedLocalFileRequestType::WritingStream,
 		[this, StreamChunkStartMS, StreamChunkEndMS, TotalLengthInMS, StreamData = MoveTemp(StreamAr.Buffer)](ELocalFileReplayResult& ReplayResult) mutable
@@ -1913,7 +1927,7 @@ void FLocalFileNetworkReplayStreamer::FlushStream(const uint32 TimeInMS)
 
 					LocalFileAr->Serialize((void*)EncryptedData.GetData(), EncryptedData.Num());
 
-					int32 ChunkSize = LocalFileAr->Tell() - MetadataPos;
+					int32 ChunkSize = IntCastChecked<int32>(LocalFileAr->Tell() - MetadataPos);
 
 					LocalFileAr->Seek(SavedPos);
 					*LocalFileAr << ChunkSize;
@@ -1972,7 +1986,7 @@ void FLocalFileNetworkReplayStreamer::FlushStream(const uint32 TimeInMS)
 
 void FLocalFileNetworkReplayStreamer::FlushCheckpoint(const uint32 TimeInMS)
 {
-	if (CheckpointAr.Buffer.Num() == 0)
+	if (CheckpointAr.TotalSize() == 0)
 	{
 		UE_LOG( LogLocalFileReplay, Warning, TEXT( "FLocalFileNetworkReplayStreamer::FlushCheckpoint. Checkpoint is empty." ) );
 		return;
@@ -1987,11 +2001,11 @@ void FLocalFileNetworkReplayStreamer::FlushCheckpoint(const uint32 TimeInMS)
 
 void FLocalFileNetworkReplayStreamer::FlushCheckpointInternal(const uint32 TimeInMS)
 {
-	if (StreamerState != EReplayStreamerState::Recording || CheckpointAr.Buffer.Num() == 0)
+	if (StreamerState != EReplayStreamerState::Recording || CheckpointAr.TotalSize() == 0)
 	{
 		// If there is no active session, or we are not recording, we don't need to flush
 		CheckpointAr.Buffer.Empty();
-		CheckpointAr.Pos = 0;
+		CheckpointAr.Seek(0);
 		return;
 	}
 
@@ -2006,88 +2020,88 @@ void FLocalFileNetworkReplayStreamer::FlushCheckpointInternal(const uint32 TimeI
 			const int32 DataChunkIndex = TaskReplayInfo.DataChunks.Num();
 			const int32 CheckpointIndex = TaskReplayInfo.Checkpoints.Num();
 
-				TSharedPtr<FArchive> LocalFileAr = CreateLocalFileWriter(GetDemoFullFilename(CurrentStreamName));
-				if (LocalFileAr.IsValid())
-				{
-					LocalFileAr->Seek(LocalFileAr->TotalSize());
+			TSharedPtr<FArchive> LocalFileAr = CreateLocalFileWriter(GetDemoFullFilename(CurrentStreamName));
+			if (LocalFileAr.IsValid())
+			{
+				LocalFileAr->Seek(LocalFileAr->TotalSize());
 
-					TArray<uint8> CompressedData;
+				TArray<uint8> CompressedData;
 
-					if (SupportsCompression())
-					{
-						SCOPE_CYCLE_COUNTER(STAT_LocalReplay_CompressTime);
+				if (SupportsCompression())
+				{
+					SCOPE_CYCLE_COUNTER(STAT_LocalReplay_CompressTime);
 
-						if (!CompressBuffer(CheckpointData, CompressedData))
-						{
-							UE_LOG(LogLocalFileReplay, Warning, TEXT("FLocalFileNetworkReplayStreamer::FlushStream - CompressBuffer failed"));
-						ReplayResult = ELocalFileReplayResult::CompressBuffer;
-							return;
-						}
-					}
-					else
+					if (!CompressBuffer(CheckpointData, CompressedData))
 					{
-						CompressedData = MoveTemp(CheckpointData);
+						UE_LOG(LogLocalFileReplay, Warning, TEXT("FLocalFileNetworkReplayStreamer::FlushStream - CompressBuffer failed"));
+						ReplayResult = ELocalFileReplayResult::CompressBuffer;
+						return;
 					}
+				}
+				else
+				{
+					CompressedData = MoveTemp(CheckpointData);
+				}
 
-					TArray<uint8> EncryptedData;
+				TArray<uint8> EncryptedData;
 
-					if (AllowEncryptedWrite())
-					{
-						SCOPE_CYCLE_COUNTER(STAT_LocalReplay_EncryptTime);
+				if (AllowEncryptedWrite())
+				{
+					SCOPE_CYCLE_COUNTER(STAT_LocalReplay_EncryptTime);
 
 					if (!EncryptBuffer(CompressedData, EncryptedData, TaskReplayInfo.EncryptionKey))
-						{
-							UE_LOG(LogLocalFileReplay, Warning, TEXT("FLocalFileNetworkReplayStreamer::FlushStream - EncryptBuffer failed"));
-						ReplayResult = ELocalFileReplayResult::EncryptBuffer;
-							return;
-						}
-					}
-					else
 					{
-						EncryptedData = MoveTemp(CompressedData);
+						UE_LOG(LogLocalFileReplay, Warning, TEXT("FLocalFileNetworkReplayStreamer::FlushStream - EncryptBuffer failed"));
+						ReplayResult = ELocalFileReplayResult::EncryptBuffer;
+						return;
 					}
+				}
+				else
+				{
+					EncryptedData = MoveTemp(CompressedData);
+				}
 
-					// flush checkpoint
-					if (EncryptedData.Num() > 0)
-					{
+				// flush checkpoint
+				if (EncryptedData.Num() > 0)
+				{
 					const int64 TypeOffset = LocalFileAr->Tell();
 
-						ELocalFileChunkType ChunkType = ELocalFileChunkType::Checkpoint;
-						*LocalFileAr << ChunkType;
+					ELocalFileChunkType ChunkType = ELocalFileChunkType::Checkpoint;
+					*LocalFileAr << ChunkType;
 
 					const int64 SavedPos = LocalFileAr->Tell();
 
-						int32 PlaceholderSize = 0;
-						*LocalFileAr << PlaceholderSize;
+					int32 PlaceholderSize = 0;
+					*LocalFileAr << PlaceholderSize;
 
 					const int64 MetadataPos = LocalFileAr->Tell();
 
-						FString Id = FString::Printf(TEXT("checkpoint%ld"), CheckpointIndex);
-						*LocalFileAr << Id;
+					FString Id = FString::Printf(TEXT("checkpoint%ld"), CheckpointIndex);
+					*LocalFileAr << Id;
 
-						FString Group = TEXT("checkpoint");
-						*LocalFileAr << Group;
+					FString Group = TEXT("checkpoint");
+					*LocalFileAr << Group;
 
-						FString Metadata = FString::Printf(TEXT("%ld"), DataChunkIndex);
-						*LocalFileAr << Metadata;
+					FString Metadata = FString::Printf(TEXT("%ld"), DataChunkIndex);
+					*LocalFileAr << Metadata;
 
-						uint32 Time1 = CheckpointTimeInMS;
-						*LocalFileAr << Time1;
+					uint32 Time1 = CheckpointTimeInMS;
+					*LocalFileAr << Time1;
 
-						uint32 Time2 = CheckpointTimeInMS;
-						*LocalFileAr << Time2;
+					uint32 Time2 = CheckpointTimeInMS;
+					*LocalFileAr << Time2;
 
-						int32 CheckpointSize = EncryptedData.Num();
-						*LocalFileAr << CheckpointSize;
+					int32 CheckpointSize = EncryptedData.Num();
+					*LocalFileAr << CheckpointSize;
 
 					const int64 EventDataOffset = LocalFileAr->Tell();
 
-						LocalFileAr->Serialize((void*)EncryptedData.GetData(), EncryptedData.Num());
+					LocalFileAr->Serialize((void*)EncryptedData.GetData(), EncryptedData.Num());
 
-						int32 ChunkSize = LocalFileAr->Tell() - MetadataPos;
+					int32 ChunkSize = IntCastChecked<int32>(LocalFileAr->Tell() - MetadataPos);
 
-						LocalFileAr->Seek(SavedPos);
-						*LocalFileAr << ChunkSize;
+					LocalFileAr->Seek(SavedPos);
+					*LocalFileAr << ChunkSize;
 
 					FLocalFileChunkInfo& ChunkInfo = TaskReplayInfo.Chunks.AddDefaulted_GetRef();
 					ChunkInfo.ChunkType = ChunkType;
@@ -2104,10 +2118,10 @@ void FLocalFileNetworkReplayStreamer::FlushCheckpointInternal(const uint32 TimeI
 					CheckpointInfo.Time2 = Time2;
 					CheckpointInfo.SizeInBytes = CheckpointSize;
 					CheckpointInfo.EventDataOffset = EventDataOffset;
-					}
-
-					LocalFileAr = nullptr;
 				}
+
+				LocalFileAr = nullptr;
+			}
 			else
 			{
 				ReplayResult = ELocalFileReplayResult::FileWriter;
@@ -2132,7 +2146,7 @@ void FLocalFileNetworkReplayStreamer::FlushCheckpointInternal(const uint32 TimeI
 		});
 
 	CheckpointAr.Buffer.Reset();
-	CheckpointAr.Pos = 0;	
+	CheckpointAr.Seek(0);
 }
 
 void FLocalFileNetworkReplayStreamer::GotoCheckpointIndex(const int32 CheckpointIndex, const FGotoCallback& Delegate, EReplayCheckpointType CheckpointType)
@@ -2157,7 +2171,7 @@ void FLocalFileNetworkReplayStreamer::GotoCheckpointIndex(const int32 Checkpoint
 				// Make sure to reset the checkpoint archive (this is how we signify that the engine should start from the beginning of the stream (we don't need a checkpoint for that))
 				CheckpointAr.Reset();
 
-				if (!IsDataAvailableForTimeRange(0, LastGotoTimeInMS))
+				if (!IsDataAvailableForTimeRange(0, IntCastChecked<uint32>(LastGotoTimeInMS)))
 				{
 					// Completely reset our stream (we're going to start loading from the start of the checkpoint)
 					StreamAr.Buffer.Reset();
@@ -2172,10 +2186,10 @@ void FLocalFileNetworkReplayStreamer::GotoCheckpointIndex(const int32 Checkpoint
 
 					LastChunkTime = 0;		// Force the next chunk to start loading immediately in case LastGotoTimeInMS is 0 (which would effectively disable high priority mode immediately)
 
-					SetHighPriorityTimeRange(0, LastGotoTimeInMS);
+					SetHighPriorityTimeRange(0, IntCastChecked<uint32>(LastGotoTimeInMS));
 				}
 
-				StreamAr.Pos = 0;
+				StreamAr.Seek(0);
 				StreamAr.bAtEndOfReplay	= false;
 
 				FGotoResult Result;
@@ -2327,9 +2341,9 @@ void FLocalFileNetworkReplayStreamer::GotoCheckpointIndex(const int32 Checkpoint
 			}
 
 			CheckpointAr.Buffer = MoveTemp(RequestData.DataBuffer);
-			CheckpointAr.Pos = 0;
+			CheckpointAr.Seek(0);
 
-			const bool bIsDataAvailableForTimeRange = IsDataAvailableForTimeRange(CurrentReplayInfo.Checkpoints[CheckpointIndex].Time1, LastGotoTimeInMS);
+			const bool bIsDataAvailableForTimeRange = IsDataAvailableForTimeRange(CurrentReplayInfo.Checkpoints[CheckpointIndex].Time1, IntCastChecked<uint32>(LastGotoTimeInMS));
 
 			const int32 DataChunkIndex = FCString::Atoi(*CurrentReplayInfo.Checkpoints[CheckpointIndex].Metadata);
 			
@@ -2356,8 +2370,8 @@ void FLocalFileNetworkReplayStreamer::GotoCheckpointIndex(const int32 Checkpoint
 				else
 				{
 					// set stream position back to the correct location
-					StreamAr.Pos = CurrentReplayInfo.DataChunks[DataChunkIndex].StreamOffset - StreamDataOffset;
-					check(StreamAr.Pos >= 0 && StreamAr.Pos <= StreamAr.Buffer.Num());
+					StreamAr.Seek(CurrentReplayInfo.DataChunks[DataChunkIndex].StreamOffset - StreamDataOffset);
+					check(StreamAr.Tell() >= 0 && StreamAr.Tell() <= StreamAr.TotalSize());
 					StreamAr.bAtEndOfReplay = false;
 				}
 			}
@@ -2381,7 +2395,7 @@ void FLocalFileNetworkReplayStreamer::GotoCheckpointIndex(const int32 Checkpoint
 			if (LastGotoTimeInMS >= 0)
 			{
 				// If we are fine scrubbing, make sure to wait on the part of the stream that is needed to do this in one frame
-				SetHighPriorityTimeRange(CurrentReplayInfo.Checkpoints[CheckpointIndex].Time1, LastGotoTimeInMS);
+				SetHighPriorityTimeRange(CurrentReplayInfo.Checkpoints[CheckpointIndex].Time1, IntCastChecked<uint32>(LastGotoTimeInMS));
 
 				// Subtract off starting time so we pass in the leftover to the engine to fast forward through for the fine scrubbing part
 				LastGotoTimeInMS -= CurrentReplayInfo.Checkpoints[CheckpointIndex].Time1;
@@ -2502,14 +2516,14 @@ void FLocalFileNetworkReplayStreamer::GotoCheckpointIndex(const int32 Checkpoint
 			AddRequestToCache(CurrentReplayInfo.Checkpoints[CheckpointIndex].ChunkIndex, RequestData.DataBuffer);
 
 			CheckpointAr.Buffer = MoveTemp(RequestData.DataBuffer);
-			CheckpointAr.Pos = 0;
+			CheckpointAr.Seek(0);
 
 			const FLocalFileEventInfo& Checkpoint = CurrentReplayInfo.Checkpoints[CheckpointIndex];
 			const int32 DataChunkIndex = FCString::Atoi(*Checkpoint.Metadata);
 
 			if (CurrentReplayInfo.DataChunks.IsValidIndex(DataChunkIndex))
 			{
-				bool bIsDataAvailableForTimeRange = IsDataAvailableForTimeRange(Checkpoint.Time1, LastGotoTimeInMS);
+				bool bIsDataAvailableForTimeRange = IsDataAvailableForTimeRange(Checkpoint.Time1, IntCastChecked<uint32>(LastGotoTimeInMS));
 
 				if (!bIsDataAvailableForTimeRange)
 				{
@@ -2532,8 +2546,8 @@ void FLocalFileNetworkReplayStreamer::GotoCheckpointIndex(const int32 Checkpoint
 				else
 				{
 					// set stream position back to the correct location
-					StreamAr.Pos = CurrentReplayInfo.DataChunks[DataChunkIndex].StreamOffset - StreamDataOffset;
-					check(StreamAr.Pos >= 0 && StreamAr.Pos <= StreamAr.Buffer.Num());
+					StreamAr.Seek(CurrentReplayInfo.DataChunks[DataChunkIndex].StreamOffset - StreamDataOffset);
+					check(StreamAr.Tell() >= 0 && StreamAr.Tell() <= StreamAr.TotalSize());
 					StreamAr.bAtEndOfReplay = false;
 				}
 			}
@@ -2549,7 +2563,7 @@ void FLocalFileNetworkReplayStreamer::GotoCheckpointIndex(const int32 Checkpoint
 			if (LastGotoTimeInMS >= 0)
 			{
 				// If we are fine scrubbing, make sure to wait on the part of the stream that is needed to do this in one frame
-				SetHighPriorityTimeRange(Checkpoint.Time1, LastGotoTimeInMS);
+				SetHighPriorityTimeRange(Checkpoint.Time1, IntCastChecked<uint32>(LastGotoTimeInMS));
 
 				// Subtract off starting time so we pass in the leftover to the engine to fast forward through for the fine scrubbing part
 				LastGotoTimeInMS -= Checkpoint.Time1;
@@ -2736,7 +2750,7 @@ const TArray<uint8>& FLocalFileNetworkReplayStreamer::GetCachedFileContents(cons
 		FArchive* Ar = IFileManager::Get().CreateFileReader(*Filename, FILEREAD_AllowWrite);
 		if (Ar)
 		{
-			Data.AddUninitialized(Ar->TotalSize());
+			Data.AddUninitialized(IntCastChecked<int32>(Ar->TotalSize()));
 			Ar->Serialize(Data.GetData(), Data.Num());
 			delete Ar;
 		}
@@ -3069,7 +3083,7 @@ void FLocalFileNetworkReplayStreamer::DownloadHeader(const FDownloadHeaderCallba
 			[this, Delegate](TLocalFileRequestCommonData<FDownloadHeaderResult>& RequestData)
 			{
 				HeaderAr.Buffer = MoveTemp(RequestData.DataBuffer);
-				HeaderAr.Pos = 0;
+				HeaderAr.Seek(0);
 
 				Delegate.ExecuteIfBound(RequestData.DelegateResult);
 			});
@@ -3091,7 +3105,7 @@ void FLocalFileNetworkReplayStreamer::WriteHeader()
 		return;
 	}
 
-	if (HeaderAr.Buffer.Num() == 0)
+	if (HeaderAr.TotalSize() == 0)
 	{
 		// Header wasn't serialized
 		UE_LOG(LogLocalFileReplay, Warning, TEXT("FLocalFileNetworkReplayStreamer::WriteHeader. No header to upload"));
@@ -3147,7 +3161,7 @@ void FLocalFileNetworkReplayStreamer::WriteHeader()
 					}
 				}
 
-				const int32 TypeOffset = LocalFileAr->Tell();
+				const int32 TypeOffset = IntCastChecked<int32>(LocalFileAr->Tell());
 
 				ELocalFileChunkType ChunkType = ELocalFileChunkType::Header;
 				*LocalFileAr << ChunkType;
@@ -3155,7 +3169,7 @@ void FLocalFileNetworkReplayStreamer::WriteHeader()
 				int32 ChunkSize = HeaderData.Num();
 				*LocalFileAr << ChunkSize;
 
-				const int32 DataOffset = LocalFileAr->Tell();
+				const int32 DataOffset = IntCastChecked<int32>(LocalFileAr->Tell());
 
 				LocalFileAr->Serialize((void*)HeaderData.GetData(), HeaderData.Num());
 				LocalFileAr = nullptr;
@@ -3295,10 +3309,10 @@ void FLocalFileNetworkReplayStreamer::ConditionallyLoadNextChunk()
 			return;		
 		}
 
-		if ((StreamTimeRange.Max > StreamTimeRange.Min) && (StreamAr.Buffer.Num() > 0))
+		if ((StreamTimeRange.Max > StreamTimeRange.Min) && (StreamAr.TotalSize() > 0))
 		{
 			// Make a guess on how far we're in
-			const float PercentIn		= StreamAr.Buffer.Num() > 0 ? (float)StreamAr.Pos / (float)StreamAr.Buffer.Num() : 0.0f;
+			const float PercentIn		= StreamAr.TotalSize() > 0 ? (float)StreamAr.Tell() / (float)StreamAr.TotalSize() : 0.0f;
 			const float TotalStreamTimeSeconds = (float)(StreamTimeRange.Size()) / 1000.0f;
 			const float CurrentTime		= TotalStreamTimeSeconds * PercentIn;
 			const float TimeLeft		= TotalStreamTimeSeconds - CurrentTime;
@@ -3422,7 +3436,7 @@ void FLocalFileNetworkReplayStreamer::ConditionallyLoadNextChunk()
 
 			if (RequestData.DataBuffer.Num() > 0)
 			{
-				if (StreamAr.Buffer.Num() == 0)
+				if (StreamAr.TotalSize() == 0)
 				{
 					StreamTimeRange.Min = CurrentReplayInfo.DataChunks[RequestedStreamChunkIndex].Time1;
 				}
@@ -3433,21 +3447,21 @@ void FLocalFileNetworkReplayStreamer::ConditionallyLoadNextChunk()
 				check(StreamTimeRange.IsValid());
 
 				// make space before appending
-				int32 MaxBufferedChunks = UE::Net::LocalFileReplay::CVarMaxBufferedStreamChunks.GetValueOnAnyThread();
+				const int32 MaxBufferedChunks = UE::Net::LocalFileReplay::CVarMaxBufferedStreamChunks.GetValueOnAnyThread();
 				if (MaxBufferedChunks > 0)
 				{
-					int32 MinChunkIndex = FMath::Max(0, (RequestedStreamChunkIndex + 1) - MaxBufferedChunks);
+					const int32 MinChunkIndex = FMath::Max(0, (RequestedStreamChunkIndex + 1) - MaxBufferedChunks);
 					if (MinChunkIndex > 0)
 					{
-						int32 TrimBytes = CurrentReplayInfo.DataChunks[MinChunkIndex].StreamOffset - StreamDataOffset;
+						const int32 TrimBytes = IntCastChecked<int32>(CurrentReplayInfo.DataChunks[MinChunkIndex].StreamOffset - StreamDataOffset);
 						if (TrimBytes > 0)
 						{
 							// can't remove chunks if we're actively seeking within that data
-							if (StreamAr.Pos >= TrimBytes)
+							if (StreamAr.Tell() >= TrimBytes)
 							{
 								// don't realloc, we're about to append anyway
 								StreamAr.Buffer.RemoveAt(0, TrimBytes, false);
-								StreamAr.Pos -= TrimBytes;
+								StreamAr.Seek(StreamAr.Tell() - TrimBytes);
 
 								StreamTimeRange.Min = CurrentReplayInfo.DataChunks[MinChunkIndex].Time1;
 								StreamDataOffset += TrimBytes;
@@ -3748,7 +3762,9 @@ void FLocalFileNetworkReplayStreamingFactory::Flush()
 				break;
 			}
 
-			Tick(AppTime - LastTime);
+			const float DeltaTime = FloatCastChecked<float>(AppTime - LastTime, UE::LWC::DefaultFloatPrecision);
+			Tick(DeltaTime);
+
 			LastTime = AppTime;
 
 			if (HasAnyPendingRequests())
diff --git a/Engine/Source/Runtime/NetworkReplayStreaming/LocalFileNetworkReplayStreaming/Public/LocalFileNetworkReplayStreaming.h b/Engine/Source/Runtime/NetworkReplayStreaming/LocalFileNetworkReplayStreaming/Public/LocalFileNetworkReplayStreaming.h
--- a/Engine/Source/Runtime/NetworkReplayStreaming/LocalFileNetworkReplayStreaming/Public/LocalFileNetworkReplayStreaming.h
+++ b/Engine/Source/Runtime/NetworkReplayStreaming/LocalFileNetworkReplayStreaming/Public/LocalFileNetworkReplayStreaming.h
@@ -175,7 +175,9 @@ struct FLocalFileReplayInfo
 class LOCALFILENETWORKREPLAYSTREAMING_API FLocalFileStreamFArchive : public FArchive
 {
 public:
+	PRAGMA_DISABLE_DEPRECATION_WARNINGS
 	FLocalFileStreamFArchive() : Pos(0), bAtEndOfReplay(false) {}
+	PRAGMA_ENABLE_DEPRECATION_WARNINGS
 
 	virtual void	Serialize(void* V, int64 Length) override;
 	virtual int64	Tell() override;
@@ -183,16 +185,26 @@ public:
 	virtual void	Seek(int64 InPos) override;
 	virtual bool	AtEnd() override;
 
+	int64 Tell() const;
+	int64 TotalSize() const;
+
 	void Reset()
 	{
 		Buffer.Reset();
+		PRAGMA_DISABLE_DEPRECATION_WARNINGS
 		Pos = 0;
+		PRAGMA_ENABLE_DEPRECATION_WARNINGS
 		bAtEndOfReplay = false;
+		ArchivePos = 0;
 	}
 
 	TArray<uint8>	Buffer;
+	UE_DEPRECATED(5.3, "Please use Tell/Seek instead")
 	int32			Pos;
 	bool			bAtEndOfReplay;
+
+private:
+	int64 ArchivePos = 0;
 };
 
 namespace EQueuedLocalFileRequestType
diff --git a/Engine/Source/Runtime/NetworkReplayStreaming/NetworkReplayStreaming/NetworkReplayStreaming.Build.cs b/Engine/Source/Runtime/NetworkReplayStreaming/NetworkReplayStreaming/NetworkReplayStreaming.Build.cs
--- a/Engine/Source/Runtime/NetworkReplayStreaming/NetworkReplayStreaming/NetworkReplayStreaming.Build.cs
+++ b/Engine/Source/Runtime/NetworkReplayStreaming/NetworkReplayStreaming/NetworkReplayStreaming.Build.cs
@@ -15,6 +15,8 @@ namespace UnrealBuildTool.Rules
 					"NetCore",
 				}
 			);
+
+			UnsafeTypeCastWarningLevel = WarningLevel.Error;
 		}
 	}
 }
diff --git a/Engine/Source/Runtime/NetworkReplayStreaming/NullNetworkReplayStreaming/NullNetworkReplayStreaming.build.cs b/Engine/Source/Runtime/NetworkReplayStreaming/NullNetworkReplayStreaming/NullNetworkReplayStreaming.build.cs
--- a/Engine/Source/Runtime/NetworkReplayStreaming/NullNetworkReplayStreaming/NullNetworkReplayStreaming.build.cs
+++ b/Engine/Source/Runtime/NetworkReplayStreaming/NullNetworkReplayStreaming/NullNetworkReplayStreaming.build.cs
@@ -17,6 +17,8 @@ namespace UnrealBuildTool.Rules
                     "Json",
 				}
 			);
+
+			UnsafeTypeCastWarningLevel = WarningLevel.Error;
 		}
 	}
 }
diff --git a/Engine/Source/Runtime/NetworkReplayStreaming/NullNetworkReplayStreaming/Private/NullNetworkReplayStreaming.cpp b/Engine/Source/Runtime/NetworkReplayStreaming/NullNetworkReplayStreaming/Private/NullNetworkReplayStreaming.cpp
--- a/Engine/Source/Runtime/NetworkReplayStreaming/NullNetworkReplayStreaming/Private/NullNetworkReplayStreaming.cpp
+++ b/Engine/Source/Runtime/NetworkReplayStreaming/NullNetworkReplayStreaming/Private/NullNetworkReplayStreaming.cpp
@@ -733,7 +733,7 @@ void FNullNetworkReplayStreamer::Tick(float DeltaSeconds)
 		// Note the size of the file between demo frames
 		if (FileAr.IsValid() && ReplayInfo.DemoFileLastOffset < FileAr->Tell())
 		{
-			ReplayInfo.DemoFileLastOffset = FileAr->Tell();
+			ReplayInfo.DemoFileLastOffset = IntCastChecked<int32>(FileAr->Tell());
 			FileAr->Flush();
 
 			WriteReplayInfo(CurrentStreamName, ReplayInfo);
diff --git a/Engine/Source/Runtime/NetworkReplayStreaming/SaveGameNetworkReplayStreaming/Private/SaveGameNetworkReplayStreaming.cpp b/Engine/Source/Runtime/NetworkReplayStreaming/SaveGameNetworkReplayStreaming/Private/SaveGameNetworkReplayStreaming.cpp
--- a/Engine/Source/Runtime/NetworkReplayStreaming/SaveGameNetworkReplayStreaming/Private/SaveGameNetworkReplayStreaming.cpp
+++ b/Engine/Source/Runtime/NetworkReplayStreaming/SaveGameNetworkReplayStreaming/Private/SaveGameNetworkReplayStreaming.cpp
@@ -1032,7 +1032,7 @@ void FSaveGameNetworkReplayStreamer::KeepReplay_Internal(const FString& ReplayNa
 		{
 			// Read in the Replay and MetaData
 			TArray<uint8> ReplayData;
-			ReplayData.SetNumUninitialized(ReplayFileAr->TotalSize());
+			ReplayData.SetNumUninitialized(IntCastChecked<int32>(ReplayFileAr->TotalSize()));
 			ReplayFileAr->Serialize(ReplayData.GetData(), ReplayData.Num());
 
 			// Now, create the Meta Data and save that.
diff --git a/Engine/Source/Runtime/NetworkReplayStreaming/SaveGameNetworkReplayStreaming/SaveGameNetworkReplayStreaming.Build.cs b/Engine/Source/Runtime/NetworkReplayStreaming/SaveGameNetworkReplayStreaming/SaveGameNetworkReplayStreaming.Build.cs
--- a/Engine/Source/Runtime/NetworkReplayStreaming/SaveGameNetworkReplayStreaming/SaveGameNetworkReplayStreaming.Build.cs
+++ b/Engine/Source/Runtime/NetworkReplayStreaming/SaveGameNetworkReplayStreaming/SaveGameNetworkReplayStreaming.Build.cs
@@ -21,6 +21,8 @@ namespace UnrealBuildTool.Rules
 					"Json"
 				}
             );
-        }
+
+			UnsafeTypeCastWarningLevel = WarningLevel.Error;
+		}
     }
 }
