Avoid check if replay header is corrupt

diff --git a/Engine/Source/Runtime/NetworkReplayStreaming/LocalFileNetworkReplayStreaming/Private/LocalFileNetworkReplayStreaming.cpp b/Engine/Source/Runtime/NetworkReplayStreaming/LocalFileNetworkReplayStreaming/Private/LocalFileNetworkReplayStreaming.cpp
--- a/Engine/Source/Runtime/NetworkReplayStreaming/LocalFileNetworkReplayStreaming/Private/LocalFileNetworkReplayStreaming.cpp
+++ b/Engine/Source/Runtime/NetworkReplayStreaming/LocalFileNetworkReplayStreaming/Private/LocalFileNetworkReplayStreaming.cpp
@@ -61,6 +61,8 @@ namespace UE::Net::LocalFileReplay
 
 	TAutoConsoleVariable<int32> CVarReplayRecordingMinSpace(TEXT("localReplay.ReplayRecordingMinSpace"), 20 * (1024 * 1024), TEXT("Minimum space needed to start recording a replay."));
 	TAutoConsoleVariable<float> CVarMinLoadNextChunkDelaySeconds(TEXT("localReplay.MinLoadNextChunkDelaySeconds"), 3.0f, TEXT("Minimum time to wait between conditional chunk loads."));
+
+	constexpr int32 MaxEncryptionKeySizeBytes = 4096;
 };
 
 const uint32 FLocalFileNetworkReplayStreamer::FileMagic = 0x1CA2E27F;
@@ -348,7 +350,23 @@ bool FLocalFileNetworkReplayStreamer::ReadReplayInfo(FArchive& Archive, FLocalFi
 
 				Info.bEncrypted = (Encrypted != 0);
 
-				Archive << Info.EncryptionKey;
+				const int64 KeyPos = Archive.Tell();
+
+				int32 KeySize;
+				Archive << KeySize;
+
+				if (KeySize >= 0 && KeySize <= LocalFileReplay::MaxEncryptionKeySizeBytes)
+				{
+					Archive.Seek(KeyPos);
+
+					Archive << Info.EncryptionKey;
+				}
+				else
+				{
+					UE_LOG(LogLocalFileReplay, Error, TEXT("ReadReplayInfo: Serialized an invalid encryption key size: %d"), KeySize);
+					Archive.SetError();
+					return false;
+				}
 			}
 
 			if (!Info.bIsLive && Info.bEncrypted && (Info.EncryptionKey.Num() == 0))
