Fixing namespace

#pf skip

diff --git a/Engine/Source/Runtime/NetworkReplayStreaming/LocalFileNetworkReplayStreaming/Private/LocalFileNetworkReplayStreaming.cpp b/Engine/Source/Runtime/NetworkReplayStreaming/LocalFileNetworkReplayStreaming/Private/LocalFileNetworkReplayStreaming.cpp
--- a/Engine/Source/Runtime/NetworkReplayStreaming/LocalFileNetworkReplayStreaming/Private/LocalFileNetworkReplayStreaming.cpp
+++ b/Engine/Source/Runtime/NetworkReplayStreaming/LocalFileNetworkReplayStreaming/Private/LocalFileNetworkReplayStreaming.cpp
@@ -355,7 +355,7 @@ bool FLocalFileNetworkReplayStreamer::ReadReplayInfo(FArchive& Archive, FLocalFi
 				int32 KeySize;
 				Archive << KeySize;
 
-				if (KeySize >= 0 && KeySize <= LocalFileReplay::MaxEncryptionKeySizeBytes)
+				if (KeySize >= 0 && KeySize <= UE::Net::LocalFileReplay::MaxEncryptionKeySizeBytes)
 				{
 					Archive.Seek(KeyPos);
 
