Adding a Flush implementation to http streamer

diff --git a/Engine/Source/Runtime/NetworkReplayStreaming/HttpNetworkReplayStreaming/Private/HttpNetworkReplayStreaming.cpp b/Engine/Source/Runtime/NetworkReplayStreaming/HttpNetworkReplayStreaming/Private/HttpNetworkReplayStreaming.cpp
--- a/Engine/Source/Runtime/NetworkReplayStreaming/HttpNetworkReplayStreaming/Private/HttpNetworkReplayStreaming.cpp
+++ b/Engine/Source/Runtime/NetworkReplayStreaming/HttpNetworkReplayStreaming/Private/HttpNetworkReplayStreaming.cpp
@@ -14,6 +14,9 @@
 #include "Engine/LocalPlayer.h"
 #include "Serialization/MemoryReader.h"
 #include "Serialization/MemoryWriter.h"
+#include "HttpManager.h"
+#include "Async/TaskGraphInterfaces.h"
+#include "Misc/Fork.h"
 
 #include UE_INLINE_GENERATED_CPP_BY_NAME(HttpNetworkReplayStreaming)
 
@@ -2735,6 +2738,72 @@ TStatId FHttpNetworkReplayStreamingFactory::GetStatId() const
 	RETURN_QUICK_DECLARE_CYCLE_STAT( FHttpNetworkReplayStreamingFactory, STATGROUP_Tickables );
 }
 
+bool FHttpNetworkReplayStreamingFactory::HasPendingHttpRequests() const
+{
+	bool bPendingRequests = false;
+
+	for (const TSharedPtr<FHttpNetworkReplayStreamer>& Streamer : HttpStreamers)
+	{
+		if (Streamer.IsValid() && Streamer->HasPendingHttpRequests())
+		{
+			bPendingRequests = true;
+			break;
+		}
+	}
+
+	return bPendingRequests;
+}
+
+void FHttpNetworkReplayStreamingFactory::Flush()
+{
+	// defaulting to false as that was the previous behavior
+	bool bFlushStreamersOnShutdown = false;
+	GConfig->GetBool(TEXT("HttpNetworkReplayStreamingFactory"), TEXT("bFlushStreamersOnShutdown"), bFlushStreamersOnShutdown, GEngineIni);
+
+	if (bFlushStreamersOnShutdown)
+	{
+		double MaxFlushTimeSeconds = -1.0;
+		GConfig->GetDouble(TEXT("HttpNetworkReplayStreamingFactory"), TEXT("MaxFlushTimeSeconds"), MaxFlushTimeSeconds, GEngineIni);
+
+		const double BeginWaitTime = FPlatformTime::Seconds();
+		double LastTime = BeginWaitTime;
+		while (HasPendingHttpRequests())
+		{
+			const double AppTime = FPlatformTime::Seconds();
+			const double TotalWait = AppTime - BeginWaitTime;
+
+			if ((MaxFlushTimeSeconds > 0) && (TotalWait > MaxFlushTimeSeconds))
+			{
+				UE_LOG(LogHttpReplay, Display, TEXT("Abandoning streamer flush after waiting %0.2f seconds"), TotalWait);
+				break;
+			}
+
+			Tick(AppTime - LastTime);
+			LastTime = AppTime;
+
+			if (HasPendingHttpRequests())
+			{
+				const float SleepInterval = 0.05f;
+
+				FHttpModule::Get().GetHttpManager().Flush(EHttpFlushReason::Default);
+
+				FTaskGraphInterface::Get().ProcessThreadUntilIdle(ENamedThreads::GameThread);
+
+				if (FPlatformProcess::SupportsMultithreading() || FForkProcessHelper::SupportsMultithreadingPostFork())
+				{
+					UE_LOG(LogHttpReplay, Display, TEXT("Sleeping to wait for outstanding requests."));
+					FPlatformProcess::Sleep(SleepInterval);
+				}
+			}
+		}
+	}
+}
+
+void FHttpNetworkReplayStreamingFactory::ShutdownModule()
+{
+	Flush();
+}
+
 #define CASE_EHTTPREPLAYRESULT_TO_TEXT_RET(txt) case txt: ReturnVal = TEXT(#txt); break;
 
 const TCHAR* LexToString(EHttpReplayResult Enum)
diff --git a/Engine/Source/Runtime/NetworkReplayStreaming/HttpNetworkReplayStreaming/Public/HttpNetworkReplayStreaming.h b/Engine/Source/Runtime/NetworkReplayStreaming/HttpNetworkReplayStreaming/Public/HttpNetworkReplayStreaming.h
--- a/Engine/Source/Runtime/NetworkReplayStreaming/HttpNetworkReplayStreaming/Public/HttpNetworkReplayStreaming.h
+++ b/Engine/Source/Runtime/NetworkReplayStreaming/HttpNetworkReplayStreaming/Public/HttpNetworkReplayStreaming.h
@@ -407,7 +407,12 @@ public:
 	virtual void Tick( float DeltaTime ) override;
 	virtual ETickableTickType GetTickableTickType() const override { return ETickableTickType::Always; }
 	virtual TStatId GetStatId() const override;
-	bool IsTickableWhenPaused() const override { return true; }
+	virtual bool IsTickableWhenPaused() const override { return true; }
+	virtual void Flush() override;
+
+	virtual void ShutdownModule() override;
+
+	bool HasPendingHttpRequests() const;
 
 	TArray< TSharedPtr< FHttpNetworkReplayStreamer > > HttpStreamers;
 };
