Fix for conversion warning

diff --git a/Engine/Source/Runtime/NetworkReplayStreaming/HttpNetworkReplayStreaming/Private/HttpNetworkReplayStreaming.cpp b/Engine/Source/Runtime/NetworkReplayStreaming/HttpNetworkReplayStreaming/Private/HttpNetworkReplayStreaming.cpp
--- a/Engine/Source/Runtime/NetworkReplayStreaming/HttpNetworkReplayStreaming/Private/HttpNetworkReplayStreaming.cpp
+++ b/Engine/Source/Runtime/NetworkReplayStreaming/HttpNetworkReplayStreaming/Private/HttpNetworkReplayStreaming.cpp
@@ -2778,7 +2778,9 @@ void FHttpNetworkReplayStreamingFactory::Flush()
 				break;
 			}
 
-			Tick(AppTime - LastTime);
+			const float DeltaTime = FloatCastChecked<float>(AppTime - LastTime, UE::LWC::DefaultFloatPrecision);
+			Tick(DeltaTime);
+
 			LastTime = AppTime;
 
 			if (HasPendingHttpRequests())
