Cap friendly name serialization size to avoid huge allocations on corrupt data

diff --git a/Engine/Source/Runtime/NetworkReplayStreaming/LocalFileNetworkReplayStreaming/Private/LocalFileNetworkReplayStreaming.cpp b/Engine/Source/Runtime/NetworkReplayStreaming/LocalFileNetworkReplayStreaming/Private/LocalFileNetworkReplayStreaming.cpp
--- a/Engine/Source/Runtime/NetworkReplayStreaming/LocalFileNetworkReplayStreaming/Private/LocalFileNetworkReplayStreaming.cpp
+++ b/Engine/Source/Runtime/NetworkReplayStreaming/LocalFileNetworkReplayStreaming/Private/LocalFileNetworkReplayStreaming.cpp
@@ -63,6 +63,8 @@ namespace UE::Net::LocalFileReplay
 	TAutoConsoleVariable<float> CVarMinLoadNextChunkDelaySeconds(TEXT("localReplay.MinLoadNextChunkDelaySeconds"), 3.0f, TEXT("Minimum time to wait between conditional chunk loads."));
 
 	constexpr int32 MaxEncryptionKeySizeBytes = 4096;
+
+	TAutoConsoleVariable<int32> CVarMaxFriendlySerializeBytes(TEXT("localReplay.MaxFriendlySerializeBytes"), 64 * 1024, TEXT("Maximum allowed serialized bytes when reading friendly name from file header."));
 };
 
 const uint32 FLocalFileNetworkReplayStreamer::FileMagic = 0x1CA2E27F;
@@ -309,7 +311,16 @@ bool FLocalFileNetworkReplayStreamer::ReadReplayInfo(FArchive& Archive, FLocalFi
 			Archive << Info.Changelist;
 
 			FString FriendlyName;
-			Archive << FriendlyName;
+			{
+				TGuardValue<int64> MaxSerialize(Archive.ArMaxSerializeSize, LocalFileReplay::CVarMaxFriendlySerializeBytes.GetValueOnAnyThread());
+				Archive << FriendlyName;
+			}
+
+			if (Archive.IsError())
+			{
+				UE_LOG(LogLocalFileReplay, Error, TEXT("ReadReplayInfo: Failed to serialize replay friendly name."));
+				return false;
+			}
 
 			SerializationInfo.FileFriendlyName = FriendlyName;
 
@@ -322,7 +333,7 @@ bool FLocalFileNetworkReplayStreamer::ReadReplayInfo(FArchive& Archive, FLocalFi
 			{
 				// Note, don't touch the FriendlyName if this is an older replay.
 				// Users can adjust the name as necessary using GetMaxFriendlyNameSize.
-				UE_LOG(LogLocalFileReplay, Warning, TEXT("FLocalFileNetworkReplayStreamer::ReadReplayInfoInternal - Loading an old replay, friendly name length **must not** be changed."));
+				UE_LOG(LogLocalFileReplay, Warning, TEXT("ReadReplayInfo - Loading an old replay, friendly name length **must not** be changed."));
 			}
 
 			uint32 IsLive;
