Fixing namespace

diff --git a/Engine/Source/Runtime/NetworkReplayStreaming/LocalFileNetworkReplayStreaming/Private/LocalFileNetworkReplayStreaming.cpp b/Engine/Source/Runtime/NetworkReplayStreaming/LocalFileNetworkReplayStreaming/Private/LocalFileNetworkReplayStreaming.cpp
--- a/Engine/Source/Runtime/NetworkReplayStreaming/LocalFileNetworkReplayStreaming/Private/LocalFileNetworkReplayStreaming.cpp
+++ b/Engine/Source/Runtime/NetworkReplayStreaming/LocalFileNetworkReplayStreaming/Private/LocalFileNetworkReplayStreaming.cpp
@@ -312,7 +312,7 @@ bool FLocalFileNetworkReplayStreamer::ReadReplayInfo(FArchive& Archive, FLocalFi
 
 			FString FriendlyName;
 			{
-				TGuardValue<int64> MaxSerialize(Archive.ArMaxSerializeSize, LocalFileReplay::CVarMaxFriendlySerializeBytes.GetValueOnAnyThread());
+				TGuardValue<int64> MaxSerialize(Archive.ArMaxSerializeSize, UE::Net::LocalFileReplay::CVarMaxFriendlySerializeBytes.GetValueOnAnyThread());
 				Archive << FriendlyName;
 			}
 
