Fix deprecation warnings with clang

diff --git a/Engine/Source/Runtime/NetworkReplayStreaming/LocalFileNetworkReplayStreaming/Public/LocalFileNetworkReplayStreaming.h b/Engine/Source/Runtime/NetworkReplayStreaming/LocalFileNetworkReplayStreaming/Public/LocalFileNetworkReplayStreaming.h
--- a/Engine/Source/Runtime/NetworkReplayStreaming/LocalFileNetworkReplayStreaming/Public/LocalFileNetworkReplayStreaming.h
+++ b/Engine/Source/Runtime/NetworkReplayStreaming/LocalFileNetworkReplayStreaming/Public/LocalFileNetworkReplayStreaming.h
@@ -177,6 +177,8 @@ class LOCALFILENETWORKREPLAYSTREAMING_API FLocalFileStreamFArchive : public FArc
 public:
 	PRAGMA_DISABLE_DEPRECATION_WARNINGS
 	FLocalFileStreamFArchive() : Pos(0), bAtEndOfReplay(false) {}
+	FLocalFileStreamFArchive(const FLocalFileStreamFArchive&) = default;
+	FLocalFileStreamFArchive& operator=(const FLocalFileStreamFArchive&) = default;
 	PRAGMA_ENABLE_DEPRECATION_WARNINGS
 
 	virtual void	Serialize(void* V, int64 Length) override;
