Adding delegate for a local file replay being finished that contains the full file path

diff --git a/Engine/Source/Runtime/NetworkReplayStreaming/LocalFileNetworkReplayStreaming/Private/LocalFileNetworkReplayStreaming.cpp b/Engine/Source/Runtime/NetworkReplayStreaming/LocalFileNetworkReplayStreaming/Private/LocalFileNetworkReplayStreaming.cpp
--- a/Engine/Source/Runtime/NetworkReplayStreaming/LocalFileNetworkReplayStreaming/Private/LocalFileNetworkReplayStreaming.cpp
+++ b/Engine/Source/Runtime/NetworkReplayStreaming/LocalFileNetworkReplayStreaming/Private/LocalFileNetworkReplayStreaming.cpp
@@ -74,6 +74,8 @@ PRAGMA_DISABLE_DEPRECATION_WARNINGS
 const uint32 FLocalFileNetworkReplayStreamer::LatestVersion = FLocalFileReplayCustomVersion::LatestVersion;
 PRAGMA_ENABLE_DEPRECATION_WARNINGS
 
+FOnLocalFileReplayFinishedWriting FLocalFileNetworkReplayStreamer::OnReplayFinishedWriting;
+
 const FGuid FLocalFileReplayCustomVersion::Guid =  FGuid(0x95A4f03E, 0x7E0B49E4, 0xBA43D356, 0x94FF87D9);
 FCustomVersionRegistration GRegisterLocalFileReplayCustomVersion(FLocalFileReplayCustomVersion::Guid, FLocalFileReplayCustomVersion::LatestVersion, TEXT("LocalFileReplay"));
 
@@ -1056,8 +1058,14 @@ void FLocalFileNetworkReplayStreamer::StopStreaming()
 			StreamAr.Reset();
 			StreamDataOffset = 0;
 			StreamChunkIndex = 0;
-			CurrentStreamName.Empty();
 			StreamerState = EReplayStreamerState::Idle;
+
+			FString StreamName = MoveTemp(CurrentStreamName);
+			const FString FullFileName = GetDemoFullFilename(StreamName);
+
+			CurrentStreamName.Empty();
+
+			OnReplayFinishedWriting.Broadcast(StreamName, FullFileName);
 		});
 }
 
diff --git a/Engine/Source/Runtime/NetworkReplayStreaming/LocalFileNetworkReplayStreaming/Public/LocalFileNetworkReplayStreaming.h b/Engine/Source/Runtime/NetworkReplayStreaming/LocalFileNetworkReplayStreaming/Public/LocalFileNetworkReplayStreaming.h
--- a/Engine/Source/Runtime/NetworkReplayStreaming/LocalFileNetworkReplayStreaming/Public/LocalFileNetworkReplayStreaming.h
+++ b/Engine/Source/Runtime/NetworkReplayStreaming/LocalFileNetworkReplayStreaming/Public/LocalFileNetworkReplayStreaming.h
@@ -502,6 +502,8 @@ protected:
 	int32 CacheKey;
 };
 
+DECLARE_MULTICAST_DELEGATE_TwoParams(FOnLocalFileReplayFinishedWriting, const FString& /*StreamName*/, const FString& /*FullReplayFile*/);
+
 /** Local file streamer that supports playback/recording to a single file on disk */
 class LOCALFILENETWORKREPLAYSTREAMING_API FLocalFileNetworkReplayStreamer : public INetworkReplayStreamer, public TSharedFromThis<FLocalFileNetworkReplayStreamer>
 {
@@ -640,6 +642,8 @@ public:
 	/** Map of checkpoint index to cached value */
 	TMap<int32, TSharedPtr<FCachedFileRequest>> DeltaCheckpointCache;
 
+	static FOnLocalFileReplayFinishedWriting OnReplayFinishedWriting;
+
 protected:
 
 	void DeleteFinishedStream_Internal(const FString& StreamName, const int32 UserIndex, const FDeleteFinishedStreamCallback& Delegate);
