[Engine]

diff --git a/Engine/Source/Runtime/Net/Core/Public/Net/Core/Connection/StateStruct.h b/Engine/Source/Runtime/Net/Core/Public/Net/Core/Connection/StateStruct.h
--- a/Engine/Source/Runtime/Net/Core/Public/Net/Core/Connection/StateStruct.h
+++ b/Engine/Source/Runtime/Net/Core/Public/Net/Core/Connection/StateStruct.h
@@ -124,7 +124,7 @@ struct FStateConfigParms
  * Base struct used for states, which is subclassed to define/implement custom states.
  */
 USTRUCT()
-struct NETCORE_API FStateStruct
+struct FStateStruct
 {
 	friend class UStatePerObjectConfig;
 
@@ -148,7 +148,7 @@ public:
 	 *
 	 * @return	The name of this state
 	 */
-	FString GetStateName() const;
+	NETCORE_API FString GetStateName() const;
 
 
 private:
@@ -200,8 +200,8 @@ protected:
 /**
  * Base class for loading and initializing state configuration
  */
-UCLASS(config=Engine, PerObjectConfig)
-class NETCORE_API UStatePerObjectConfig : public UObject
+UCLASS(config=Engine, PerObjectConfig, MinimalAPI)
+class UStatePerObjectConfig : public UObject
 {
 	GENERATED_BODY()
 
@@ -212,7 +212,7 @@ public:
 	 * @param ConfigParms	Specifies the parameters for the config section/object to get/create, and how it should be created
 	 * @return				Returns a cached or newly created instance of the config section/object
 	 */
-	static const UStatePerObjectConfig* Get(FStateConfigParms ConfigParms);
+	static NETCORE_API const UStatePerObjectConfig* Get(FStateConfigParms ConfigParms);
 
 	/**
 	 * Applies the specified config state, to an object implementing FStateStruct's.
@@ -221,14 +221,14 @@ public:
 	 * @param ConfigState	The configuration state containing the settings to be applied
 	 * @param TargetState	The object implementing FStateStruct, to apply the config settings to.
 	 */
-	static void ApplyState(const FStructOnScope& ConfigState, FStateStruct* TargetState);
+	static NETCORE_API void ApplyState(const FStructOnScope& ConfigState, FStateStruct* TargetState);
 
 
 #if !UE_BUILD_SHIPPING
 	/**
 	 * Debug function which prints the entire state configuration to the log
 	 */
-	virtual void DebugDump() const;
+	NETCORE_API virtual void DebugDump() const;
 #endif
 
 
@@ -271,9 +271,9 @@ private:
 	}
 
 
-	void RegisterStateConfig(const TArray<FString>& StateNames, TArray<TStructOnScope<FStateStruct>>& OutStates);
+	NETCORE_API void RegisterStateConfig(const TArray<FString>& StateNames, TArray<TStructOnScope<FStateStruct>>& OutStates);
 
-	virtual void OverridePerObjectConfigSection(FString& SectionName) override;
+	NETCORE_API virtual void OverridePerObjectConfigSection(FString& SectionName) override;
 
 
 	/**
@@ -284,7 +284,7 @@ private:
 	 * @param FormattedName		The base formatted object name to use, for unique/singleton object lookup/caching
 	 * @return					Returns a cached or newly created instance of the config section/object archetype
 	 */
-	static UStatePerObjectConfig* GetArchetype(FStateConfigParms ConfigParms, FString FullSection, FString FormattedName);
+	static NETCORE_API UStatePerObjectConfig* GetArchetype(FStateConfigParms ConfigParms, FString FullSection, FString FormattedName);
 
 	/**
 	 * Uses reflection to load all struct config variables from the specified ini section.
@@ -295,7 +295,7 @@ private:
 	 * @param InFilename	The ini filename to read from
 	 * @return				Whether or not the struct config variables were read successfully
 	 */
-	static bool LoadStructConfig(FStructOnScope& OutStruct, const TCHAR* SectionName, const TCHAR* InFilename=nullptr);
+	static NETCORE_API bool LoadStructConfig(FStructOnScope& OutStruct, const TCHAR* SectionName, const TCHAR* InFilename=nullptr);
 
 
 private:
diff --git a/Engine/Source/Runtime/Net/Core/Public/Net/Core/Misc/DDoSDetection.h b/Engine/Source/Runtime/Net/Core/Public/Net/Core/Misc/DDoSDetection.h
--- a/Engine/Source/Runtime/Net/Core/Public/Net/Core/Misc/DDoSDetection.h
+++ b/Engine/Source/Runtime/Net/Core/Public/Net/Core/Misc/DDoSDetection.h
@@ -22,7 +22,7 @@ DECLARE_DELEGATE_OneParam(FDDoSSeverityEscalation, FString /*SeverityCategory*/)
 /**
  * Struct containing the per-second packet counters
  */
-struct NETCORE_API FDDoSPacketCounters
+struct FDDoSPacketCounters
 {
 	/** Counter for non-NetConnection packets received, since the last per second quota period began */
 	int32 NonConnPacketCounter;
@@ -62,7 +62,7 @@ struct NETCORE_API FDDoSPacketCounters
 /**
  * Stores the DDoS detection state (either settings from the config file, or the active DDoS detection state)
  */
-struct NETCORE_API FDDoSState
+struct FDDoSState
 {
 	/** Whether or not to send analytics when escalating to this state */
 	bool bSendEscalateAnalytics;
@@ -126,7 +126,7 @@ struct NETCORE_API FDDoSState
 /**
  * DDoS detection state, with functions for applying the state to active DDoS detection
  */
-struct NETCORE_API FDDoSStateConfig : public FDDoSState
+struct FDDoSStateConfig : public FDDoSState
 {
 	/** The name of the DDoS severity level this config section represents */
 	FString SeverityCategory;
@@ -169,13 +169,13 @@ struct NETCORE_API FDDoSStateConfig : public FDDoSState
  * The main DDoS detection tracking class, for counting packets and applying restrictions.
  * Implemented separate to the NetDriver, to allow wider use e.g. potentially at socket level, if useful.
  */
-class NETCORE_API FDDoSDetection : protected FDDoSPacketCounters, protected FDDoSState
+class FDDoSDetection : protected FDDoSPacketCounters, protected FDDoSState
 {
 public:
 	/**
 	 * Default constructor
 	 */
-	FDDoSDetection();
+	NETCORE_API FDDoSDetection();
 
 
 	/**
@@ -183,35 +183,35 @@ public:
 	 *
 	 * @param MaxTickRate	The maximum tick rate of the server
 	 */
-	void Init(int32 MaxTickRate);
+	NETCORE_API void Init(int32 MaxTickRate);
 
 	/**
 	 * Initializes the settings from the .ini file - must support reloading of settings on-the-fly
 	 */
-	void InitConfig();
+	NETCORE_API void InitConfig();
 
 	/**
 	* Initialize the expected tick rate.
 	*/
-	void SetMaxTickRate(int32 MaxTickRate);
+	NETCORE_API void SetMaxTickRate(int32 MaxTickRate);
 
 	/**
 	 * Updates the current DDoS detection severity state
 	 *
 	 * @param bEscalate		Whether or not we are escalating or de-escalating the severity state
 	 */
-	void UpdateSeverity(bool bEscalate);
+	NETCORE_API void UpdateSeverity(bool bEscalate);
 
 
 	/**
 	 * Triggered before packet receive begins, during the current frame
 	 */
-	void PreFrameReceive(float DeltaTime);
+	NETCORE_API void PreFrameReceive(float DeltaTime);
 
 	/**
 	 * Triggered after packet receive ends, during the current frame
 	 */
-	void PostFrameReceive();
+	NETCORE_API void PostFrameReceive();
 
 
 	/**
@@ -274,7 +274,7 @@ protected:
 	 *
 	 * @return	Whether or not non-NetConnection packet limits have been reached
 	 */
-	bool CheckNonConnQuotasAndLimits();
+	NETCORE_API bool CheckNonConnQuotasAndLimits();
 
 	/**
 	 * Performs periodic checks on NetConnection packet limits
diff --git a/Engine/Source/Runtime/Net/Core/Public/Net/Core/Misc/GuidReferences.h b/Engine/Source/Runtime/Net/Core/Public/Net/Core/Misc/GuidReferences.h
--- a/Engine/Source/Runtime/Net/Core/Public/Net/Core/Misc/GuidReferences.h
+++ b/Engine/Source/Runtime/Net/Core/Public/Net/Core/Misc/GuidReferences.h
@@ -37,7 +37,7 @@ using FGuidReferencesMap = TMap<int32, FGuidReferences>;
  * Guid References may also be nested in properties (like dynamic arrays), so we'll recursively track
  * those as well.
  */
-class NETCORE_API FGuidReferences
+class FGuidReferences
 {
 public:
 	FGuidReferences():
@@ -74,7 +74,7 @@ public:
 		Array(InArray)
 	{}
 
-	~FGuidReferences();
+	NETCORE_API ~FGuidReferences();
 
 	FGuidReferences& operator=(const FGuidReferences& Other) = delete;
 
@@ -131,4 +131,4 @@ public:
 	 * then this will be a valid FGuidReferencesMap containing the nested FGuidReferences.
 	 */
 	FGuidReferencesMap* Array;
-};
\ No newline at end of file
+};
diff --git a/Engine/Source/Runtime/Net/Core/Public/Net/Core/Misc/NetConditionGroupManager.h b/Engine/Source/Runtime/Net/Core/Public/Net/Core/Misc/NetConditionGroupManager.h
--- a/Engine/Source/Runtime/Net/Core/Public/Net/Core/Misc/NetConditionGroupManager.h
+++ b/Engine/Source/Runtime/Net/Core/Public/Net/Core/Misc/NetConditionGroupManager.h
@@ -46,24 +46,24 @@ namespace UE::Net
 * Groups are inclusive and non-binding so subobject's can be part of multiple groups at the same time.
 * As long as a player is a member of one of the subobject's group, that subobject will be replicated to that connection.
 */
-class NETCORE_API FNetConditionGroupManager
+class FNetConditionGroupManager
 {
 public:
 
 	/** Add a subobject into a netcondition group */
-	void RegisterSubObjectInGroup(UObject* SubObject, FName NetConditionGroup);
+	NETCORE_API void RegisterSubObjectInGroup(UObject* SubObject, FName NetConditionGroup);
 
 	/** Add a subobject to many netcondition groups */
-	void RegisterSubObjectInMultipleGroups(UObject* SubObject, TArrayView<FName> NetConditionGroups);
+	NETCORE_API void RegisterSubObjectInMultipleGroups(UObject* SubObject, TArrayView<FName> NetConditionGroups);
 
 	/** Remove a subobject from a netcondition group */
-	void UnregisterSubObjectFromGroup(UObject* SubObject, FName NetConditionGroup);
+	NETCORE_API void UnregisterSubObjectFromGroup(UObject* SubObject, FName NetConditionGroup);
 
 	/** Remove a subobject from many netcondition groups */
-	void UnregisterSubObjectFromMultipleGroups(UObject* SubObject, TArrayView<FName> NetConditionGroups);
+	NETCORE_API void UnregisterSubObjectFromMultipleGroups(UObject* SubObject, TArrayView<FName> NetConditionGroups);
 
 	/** Remove the subobject from all groups it was included into */
-	void UnregisterSubObjectFromAllGroups(UObject* SubObject);
+	NETCORE_API void UnregisterSubObjectFromAllGroups(UObject* SubObject);
 
 	/** Return the groups the subobject is a part of. Will be empty if it wasn't added to any. */
 	TArrayView<const FName> GetSubObjectNetConditionGroups(FObjectKey SubObjectKey) const
@@ -72,7 +72,7 @@ public:
 		return GroupInfo ? TArrayView<const FName>(GroupInfo->NetGroups) : TArrayView<const FName>();
 	}
 
-	void CountBytes(FArchive& Ar) const;
+	NETCORE_API void CountBytes(FArchive& Ar) const;
 
 private:
 
@@ -91,4 +91,4 @@ private:
 };
 
 
-} // end namespace UE::Net
\ No newline at end of file
+} // end namespace UE::Net
diff --git a/Engine/Source/Runtime/Net/Core/Public/Net/Core/Misc/NetContext.h b/Engine/Source/Runtime/Net/Core/Public/Net/Core/Misc/NetContext.h
--- a/Engine/Source/Runtime/Net/Core/Public/Net/Core/Misc/NetContext.h
+++ b/Engine/Source/Runtime/Net/Core/Public/Net/Core/Misc/NetContext.h
@@ -15,7 +15,7 @@ namespace UE::Net
 {
 
 /** Stateless class that provides misc network context information */
-class NETCORE_API FNetContext
+class FNetContext
 {
 public:
 	
@@ -32,12 +32,12 @@ private:
 	FNetContext() = delete;
 	~FNetContext() = delete;
 
-	static bool bIsInRPCStack;
+	static NETCORE_API bool bIsInRPCStack;
 };
 
 
 /** Used by friendly class to set the right network context */
-class NETCORE_API FScopedNetContextRPC
+class FScopedNetContextRPC
 {
 private:
 
@@ -45,8 +45,8 @@ private:
 	friend UE::Net::Private::FNetRPC;
 	friend FObjectReplicator;
 	
-	FScopedNetContextRPC();
-	~FScopedNetContextRPC();
+	NETCORE_API FScopedNetContextRPC();
+	NETCORE_API ~FScopedNetContextRPC();
 };
 
 } // end namespace UE::Net
diff --git a/Engine/Source/Runtime/Net/Core/Public/Net/Core/Misc/NetSubObjectRegistry.h b/Engine/Source/Runtime/Net/Core/Public/Net/Core/Misc/NetSubObjectRegistry.h
--- a/Engine/Source/Runtime/Net/Core/Public/Net/Core/Misc/NetSubObjectRegistry.h
+++ b/Engine/Source/Runtime/Net/Core/Public/Net/Core/Misc/NetSubObjectRegistry.h
@@ -25,7 +25,7 @@ namespace UE::Net
 /**
 * Stores the SubObjects to replicate and the network condition dictating to which connection they can replicate to.
 */
-struct NETCORE_API FSubObjectRegistry
+struct FSubObjectRegistry
 {
 public:
 
@@ -40,10 +40,10 @@ public:
 	};
 
 	/** Adds a subobject and returns if it's a new entry or existing entry */
-	FSubObjectRegistry::EResult AddSubObjectUnique(UObject* InSubObject, ELifetimeCondition InNetCondition);
+	NETCORE_API FSubObjectRegistry::EResult AddSubObjectUnique(UObject* InSubObject, ELifetimeCondition InNetCondition);
 
 	/** Remove the subobject from the replicated list. Returns true if the subobject had been registered. */
-	bool RemoveSubObject(UObject* InSubObject);
+	NETCORE_API bool RemoveSubObject(UObject* InSubObject);
 
 	bool IsEmpty() const { return Registry.Num() == 0; }
 
@@ -51,10 +51,10 @@ public:
 	* Remove all the indexes specified by the passed array
 	* @param IndexesToClean Array of indexes sorted from smallest to biggest. The sorted order is assumed to have been done by the caller
 	*/
-	void CleanRegistryIndexes(const TArrayView<int32>& IndexesToClean);
+	NETCORE_API void CleanRegistryIndexes(const TArrayView<int32>& IndexesToClean);
 
 	/** Find the NetCondition of a SubObject. Returns COND_MAX if not registered */
-	ELifetimeCondition GetNetCondition(UObject* SubObject) const;
+	NETCORE_API ELifetimeCondition GetNetCondition(UObject* SubObject) const;
 
 	struct FEntry
 	{
@@ -114,7 +114,7 @@ public:
 	const TArray<FSubObjectRegistry::FEntry>& GetRegistryList() const { return Registry; }
 
 	/** Returns true if the subobject is contained in this list */
-	bool IsSubObjectInRegistry(const UObject* SubObject) const;
+	NETCORE_API bool IsSubObjectInRegistry(const UObject* SubObject) const;
 
 private:
 
diff --git a/Engine/Source/Runtime/Net/Core/Public/Net/Core/Misc/PacketAudit.h b/Engine/Source/Runtime/Net/Core/Public/Net/Core/Misc/PacketAudit.h
--- a/Engine/Source/Runtime/Net/Core/Public/Net/Core/Misc/PacketAudit.h
+++ b/Engine/Source/Runtime/Net/Core/Public/Net/Core/Misc/PacketAudit.h
@@ -70,7 +70,7 @@ extern NETCORE_API FPacketAudit* GPacketAuditor;
  *
  * IMPORTANT: Only supports interaction between a single server, and a single client.
  */
-class NETCORE_API FPacketAudit
+class FPacketAudit
 {
 #if !UE_BUILD_SHIPPING
 protected:
@@ -99,23 +99,23 @@ private:
 	/**
 	 * Default constructor
 	 */
-	FPacketAudit(FPlatformProcess::FSemaphore* InGameMutex);
+	NETCORE_API FPacketAudit(FPlatformProcess::FSemaphore* InGameMutex);
 
 	/**
 	 * Default destructor
 	 */
-	~FPacketAudit();
+	NETCORE_API ~FPacketAudit();
 
 public:
 	/**
 	 * Initialize GPacketAudtor
 	 */
-	static void Init();
+	static NETCORE_API void Init();
 
 	/**
 	 * Destroy GPacketAuditor
 	 */
-	static void Destruct();
+	static NETCORE_API void Destruct();
 
 	/**
 	 * On the send side, marks a named stage during packet writing, which should be audited on the receive side.
@@ -211,13 +211,13 @@ public:
 #endif
 
 protected:
-	void AddStage_Internal(FString StageName, FBitWriter& OutPacket, bool bByteAligned);
+	NETCORE_API void AddStage_Internal(FString StageName, FBitWriter& OutPacket, bool bByteAligned);
 
-	void CheckStage_Internal(FString StageName, FBitReader& InPacket, bool bByteAligned);
+	NETCORE_API void CheckStage_Internal(FString StageName, FBitReader& InPacket, bool bByteAligned);
 
-	void NotifyLowLevelSend_Internal(FBitWriter& OutPacket);
+	NETCORE_API void NotifyLowLevelSend_Internal(FBitWriter& OutPacket);
 
-	void NotifyLowLevelReceive_Internal(FBitReader& InPacket);
+	NETCORE_API void NotifyLowLevelReceive_Internal(FBitReader& InPacket);
 
 
 	/**
@@ -227,13 +227,13 @@ protected:
 	 * @param DataLenBits	The length of the data, in bits
 	 * @return				Returns the CRC of the packet
 	 */
-	static uint32 PacketCRC(uint8* Data, uint32 DataLenBits);
+	static NETCORE_API uint32 PacketCRC(uint8* Data, uint32 DataLenBits);
 
 #if !UE_BUILD_SHIPPING
 	/**
 	 * Dumps the collected packet audit data to the log
 	 */
-	void DumpAuditData(FScopedAuditAccess& AuditLock, uint32 InPacketCRC=0);
+	NETCORE_API void DumpAuditData(FScopedAuditAccess& AuditLock, uint32 InPacketCRC=0);
 #endif
 
 
diff --git a/Engine/Source/Runtime/Net/Core/Public/Net/Core/PropertyConditions/PropertyConditions.h b/Engine/Source/Runtime/Net/Core/Public/Net/Core/PropertyConditions/PropertyConditions.h
--- a/Engine/Source/Runtime/Net/Core/Public/Net/Core/PropertyConditions/PropertyConditions.h
+++ b/Engine/Source/Runtime/Net/Core/Public/Net/Core/PropertyConditions/PropertyConditions.h
@@ -16,25 +16,25 @@ class FRepChangedPropertyTracker;
 namespace UE::Net::Private
 {
 
-class NETCORE_API FNetPropertyConditionManager
+class FNetPropertyConditionManager
 {
 public:
-	FNetPropertyConditionManager();
-	~FNetPropertyConditionManager();
+	NETCORE_API FNetPropertyConditionManager();
+	NETCORE_API ~FNetPropertyConditionManager();
 
-	static FNetPropertyConditionManager& Get();
+	static NETCORE_API FNetPropertyConditionManager& Get();
 
-	void SetPropertyActive(const FObjectKey ObjectKey, const uint16 RepIndex, const bool bActive);
-	void SetPropertyDynamicCondition(const FObjectKey ObjectKey, const uint16 RepIndex, const ELifetimeCondition Condition);
+	NETCORE_API void SetPropertyActive(const FObjectKey ObjectKey, const uint16 RepIndex, const bool bActive);
+	NETCORE_API void SetPropertyDynamicCondition(const FObjectKey ObjectKey, const uint16 RepIndex, const ELifetimeCondition Condition);
 
-	void NotifyObjectDestroyed(const FObjectKey ObjectKey);
+	NETCORE_API void NotifyObjectDestroyed(const FObjectKey ObjectKey);
 
-	TSharedPtr<FRepChangedPropertyTracker> FindOrCreatePropertyTracker(const FObjectKey ObjectKey);
-	TSharedPtr<FRepChangedPropertyTracker> FindPropertyTracker(const FObjectKey ObjectKey) const;
+	NETCORE_API TSharedPtr<FRepChangedPropertyTracker> FindOrCreatePropertyTracker(const FObjectKey ObjectKey);
+	NETCORE_API TSharedPtr<FRepChangedPropertyTracker> FindPropertyTracker(const FObjectKey ObjectKey) const;
 
-	void LogMemory(FOutputDevice& Ar);
+	NETCORE_API void LogMemory(FOutputDevice& Ar);
 
-	static void SetPropertyActiveOverride(IRepChangedPropertyTracker& Tracker, UObject* OwningObject, const uint16 RepIndex, const bool bIsActive);
+	static NETCORE_API void SetPropertyActiveOverride(IRepChangedPropertyTracker& Tracker, UObject* OwningObject, const uint16 RepIndex, const bool bIsActive);
 
 private:
 	void PostGarbageCollect();
@@ -44,4 +44,4 @@ private:
 	TMap<FObjectKey, TSharedPtr<FRepChangedPropertyTracker>> PropertyTrackerMap;
 };
 
-}; // UE::Net::Private
\ No newline at end of file
+}; // UE::Net::Private
diff --git a/Engine/Source/Runtime/Net/Core/Public/Net/Core/PropertyConditions/PropertyConditionsDelegates.h b/Engine/Source/Runtime/Net/Core/Public/Net/Core/PropertyConditions/PropertyConditionsDelegates.h
--- a/Engine/Source/Runtime/Net/Core/Public/Net/Core/PropertyConditions/PropertyConditionsDelegates.h
+++ b/Engine/Source/Runtime/Net/Core/Public/Net/Core/PropertyConditions/PropertyConditionsDelegates.h
@@ -15,7 +15,7 @@ enum ELifetimeCondition : int;
 namespace UE::Net::Private
 {
 
-class NETCORE_API FPropertyConditionDelegates
+class FPropertyConditionDelegates
 {
 public:
 	using FOnPropertyCustomConditionChanged = TMulticastDelegate<void(const UObject* OwningObject, const uint16 RepIndex, const bool bActive)>;
@@ -25,8 +25,8 @@ public:
 	static FOnPropertyDynamicConditionChanged& GetOnPropertyDynamicConditionChangedDelegate() { return OnPropertyDynamicConditionChangedDelegate; }
 
 private:
-	static FOnPropertyCustomConditionChanged OnPropertyCustomConditionChangedDelegate;
-	static FOnPropertyDynamicConditionChanged OnPropertyDynamicConditionChangedDelegate;
+	static NETCORE_API FOnPropertyCustomConditionChanged OnPropertyCustomConditionChangedDelegate;
+	static NETCORE_API FOnPropertyDynamicConditionChanged OnPropertyDynamicConditionChangedDelegate;
 };
 
 } // UE::Net::Private
diff --git a/Engine/Source/Runtime/Net/Core/Public/Net/Core/Trace/Private/NetTraceInternal.h b/Engine/Source/Runtime/Net/Core/Public/Net/Core/Trace/Private/NetTraceInternal.h
--- a/Engine/Source/Runtime/Net/Core/Public/Net/Core/Trace/Private/NetTraceInternal.h
+++ b/Engine/Source/Runtime/Net/Core/Public/Net/Core/Trace/Private/NetTraceInternal.h
@@ -161,7 +161,7 @@ struct FNetTrace
 	enum { InvalidEventIndex = ~0U };
 
 	/** Returns if tracing is enabled or not */
-	NETCORE_API static bool IsEnabled() { return GetTraceVerbosity() > 0U; }
+	static bool IsEnabled() { return GetTraceVerbosity() > 0U; }
 
 	/** Sets the runtime trace verbosity level, Verbosity == 0 means that it is disabled */
 	NETCORE_API static void SetTraceVerbosity(uint32 Verbosity);
diff --git a/Engine/Source/Runtime/NetworkFile/Public/NetworkPlatformFile.h b/Engine/Source/Runtime/NetworkFile/Public/NetworkPlatformFile.h
--- a/Engine/Source/Runtime/NetworkFile/Public/NetworkPlatformFile.h
+++ b/Engine/Source/Runtime/NetworkFile/Public/NetworkPlatformFile.h
@@ -35,7 +35,7 @@ DECLARE_LOG_CATEGORY_EXTERN(LogNetworkPlatformFile, Log, All);
 /**
  * Wrapper to redirect the low level file system to a server
  */
-class NETWORKFILE_API FNetworkPlatformFile : public IPlatformFile, public FSelfRegisteringExec
+class FNetworkPlatformFile : public IPlatformFile, public FSelfRegisteringExec
 {
 	friend class FAsyncFileSync;
 	friend void ReadUnsolicitedFile(int32 InNumUnsolictedFiles, FNetworkPlatformFile& InNetworkFile, IPlatformFile& InInnerPlatformFile,  FString& InServerEngineDir, FString& InServerProjectDir);
@@ -48,9 +48,9 @@ protected:
 	 * @param HostIP host IP address
 	 * @return true if the initialization succeeded, false otherwise
 	 */
-	virtual bool InitializeInternal(IPlatformFile* Inner, const TCHAR* HostIP);
+	NETWORKFILE_API virtual bool InitializeInternal(IPlatformFile* Inner, const TCHAR* HostIP);
 
-	virtual void OnFileUpdated(const FString& LocalFilename);
+	NETWORKFILE_API virtual void OnFileUpdated(const FString& LocalFilename);
 
 public:
 	
@@ -61,14 +61,14 @@ public:
 	}
 
 	/** Constructor */
-	FNetworkPlatformFile();
+	NETWORKFILE_API FNetworkPlatformFile();
 
 	/** Destructor */
-	virtual ~FNetworkPlatformFile();
+	NETWORKFILE_API virtual ~FNetworkPlatformFile();
 
-	virtual bool ShouldBeUsed(IPlatformFile* Inner, const TCHAR* CmdLine) const override;
-	virtual bool Initialize(IPlatformFile* Inner, const TCHAR* CmdLine) override;
-	virtual void InitializeAfterSetActive() override;
+	NETWORKFILE_API virtual bool ShouldBeUsed(IPlatformFile* Inner, const TCHAR* CmdLine) const override;
+	NETWORKFILE_API virtual bool Initialize(IPlatformFile* Inner, const TCHAR* CmdLine) override;
+	NETWORKFILE_API virtual void InitializeAfterSetActive() override;
 
 	virtual IPlatformFile* GetLowerLevel() override
 	{
@@ -111,22 +111,22 @@ public:
 		GetFileInfo(Filename, Info);
 		return Info.Size;
 	}
-	virtual bool		DeleteFile(const TCHAR* Filename) override;
+	NETWORKFILE_API virtual bool		DeleteFile(const TCHAR* Filename) override;
 	virtual bool		IsReadOnly(const TCHAR* Filename) override
 	{
 		FFileInfo Info;
 		GetFileInfo(Filename, Info);
 		return Info.ReadOnly;
 	}
-	virtual bool		MoveFile(const TCHAR* To, const TCHAR* From) override;
-	virtual bool		SetReadOnly(const TCHAR* Filename, bool bNewReadOnlyValue) override;
+	NETWORKFILE_API virtual bool		MoveFile(const TCHAR* To, const TCHAR* From) override;
+	NETWORKFILE_API virtual bool		SetReadOnly(const TCHAR* Filename, bool bNewReadOnlyValue) override;
 	virtual FDateTime	GetTimeStamp(const TCHAR* Filename) override
 	{
 		FFileInfo Info;
 		GetFileInfo(Filename, Info);
 		return Info.TimeStamp;
 	}
-	virtual void		SetTimeStamp(const TCHAR* Filename, FDateTime DateTime) override;
+	NETWORKFILE_API virtual void		SetTimeStamp(const TCHAR* Filename, FDateTime DateTime) override;
 	virtual FDateTime	GetAccessTimeStamp(const TCHAR* Filename) override
 	{
 		FFileInfo Info;
@@ -137,59 +137,59 @@ public:
 	{
 		return Filename;
 	}
-	virtual IFileHandle*	OpenRead(const TCHAR* Filename, bool bAllowWrite = false) override;
-	virtual IFileHandle*	OpenWrite(const TCHAR* Filename, bool bAppend = false, bool bAllowRead = false) override;
-	virtual bool		DirectoryExists(const TCHAR* Directory) override;
-	virtual bool		CreateDirectoryTree(const TCHAR* Directory) override;
-	virtual bool		CreateDirectory(const TCHAR* Directory) override;
-	virtual bool		DeleteDirectory(const TCHAR* Directory) override;
+	NETWORKFILE_API virtual IFileHandle*	OpenRead(const TCHAR* Filename, bool bAllowWrite = false) override;
+	NETWORKFILE_API virtual IFileHandle*	OpenWrite(const TCHAR* Filename, bool bAppend = false, bool bAllowRead = false) override;
+	NETWORKFILE_API virtual bool		DirectoryExists(const TCHAR* Directory) override;
+	NETWORKFILE_API virtual bool		CreateDirectoryTree(const TCHAR* Directory) override;
+	NETWORKFILE_API virtual bool		CreateDirectory(const TCHAR* Directory) override;
+	NETWORKFILE_API virtual bool		DeleteDirectory(const TCHAR* Directory) override;
 
-	virtual FFileStatData GetStatData(const TCHAR* FilenameOrDirectory) override;
+	NETWORKFILE_API virtual FFileStatData GetStatData(const TCHAR* FilenameOrDirectory) override;
 
-	virtual bool		IterateDirectory(const TCHAR* Directory, IPlatformFile::FDirectoryVisitor& Visitor) override;
-	virtual bool		IterateDirectoryRecursively(const TCHAR* Directory, IPlatformFile::FDirectoryVisitor& Visitor) override;
+	NETWORKFILE_API virtual bool		IterateDirectory(const TCHAR* Directory, IPlatformFile::FDirectoryVisitor& Visitor) override;
+	NETWORKFILE_API virtual bool		IterateDirectoryRecursively(const TCHAR* Directory, IPlatformFile::FDirectoryVisitor& Visitor) override;
 
-	virtual bool		IterateDirectoryStat(const TCHAR* Directory, IPlatformFile::FDirectoryStatVisitor& Visitor) override;
-	virtual bool		IterateDirectoryStatRecursively(const TCHAR* Directory, IPlatformFile::FDirectoryStatVisitor& Visitor) override;
+	NETWORKFILE_API virtual bool		IterateDirectoryStat(const TCHAR* Directory, IPlatformFile::FDirectoryStatVisitor& Visitor) override;
+	NETWORKFILE_API virtual bool		IterateDirectoryStatRecursively(const TCHAR* Directory, IPlatformFile::FDirectoryStatVisitor& Visitor) override;
 
-	virtual bool		DeleteDirectoryRecursively(const TCHAR* Directory) override;
-	virtual bool		CopyFile(const TCHAR* To, const TCHAR* From, EPlatformFileRead ReadFlags = EPlatformFileRead::None, EPlatformFileWrite WriteFlags = EPlatformFileWrite::None) override;
+	NETWORKFILE_API virtual bool		DeleteDirectoryRecursively(const TCHAR* Directory) override;
+	NETWORKFILE_API virtual bool		CopyFile(const TCHAR* To, const TCHAR* From, EPlatformFileRead ReadFlags = EPlatformFileRead::None, EPlatformFileWrite WriteFlags = EPlatformFileWrite::None) override;
 
-	virtual FString ConvertToAbsolutePathForExternalAppForRead( const TCHAR* Filename ) override;
-	virtual FString ConvertToAbsolutePathForExternalAppForWrite( const TCHAR* Filename ) override;
+	NETWORKFILE_API virtual FString ConvertToAbsolutePathForExternalAppForRead( const TCHAR* Filename ) override;
+	NETWORKFILE_API virtual FString ConvertToAbsolutePathForExternalAppForWrite( const TCHAR* Filename ) override;
 
-	virtual bool SendMessageToServer(const TCHAR* Message, IPlatformFile::IFileServerMessageHandler* Handler) override;
+	NETWORKFILE_API virtual bool SendMessageToServer(const TCHAR* Message, IPlatformFile::IFileServerMessageHandler* Handler) override;
 
 
-	virtual void Tick() override;
+	NETWORKFILE_API virtual void Tick() override;
 
 
 	
-	virtual bool SendPayloadAndReceiveResponse(TArray<uint8>& In, TArray<uint8>& Out);
-	virtual bool ReceiveResponse(TArray<uint8>& Out);
+	NETWORKFILE_API virtual bool SendPayloadAndReceiveResponse(TArray<uint8>& In, TArray<uint8>& Out);
+	NETWORKFILE_API virtual bool ReceiveResponse(TArray<uint8>& Out);
 
-	bool SendReadMessage(uint8* Destination, int64 BytesToRead);
-	bool SendWriteMessage(const uint8* Source, int64 BytesToWrite);
+	NETWORKFILE_API bool SendReadMessage(uint8* Destination, int64 BytesToRead);
+	NETWORKFILE_API bool SendWriteMessage(const uint8* Source, int64 BytesToWrite);
 
-	static void ConvertServerFilenameToClientFilename(FString& FilenameToConvert, const FString& InServerEngineDir, const FString& InServerProjectDir, const FString& InServerEnginePlatformExtensionsDir, const FString& InServerProjectPlatformExtensionsDir);
+	static NETWORKFILE_API void ConvertServerFilenameToClientFilename(FString& FilenameToConvert, const FString& InServerEngineDir, const FString& InServerProjectDir, const FString& InServerEnginePlatformExtensionsDir, const FString& InServerProjectPlatformExtensionsDir);
 
 
-	virtual FString GetVersionInfo() const;
+	NETWORKFILE_API virtual FString GetVersionInfo() const;
 
 protected:
 
 	//////////////////////////////////////////////////////////////////////////
 	// FSelfRegisteringExec interface
-	virtual bool Exec_Runtime(class UWorld* InWorld, const TCHAR* Cmd, FOutputDevice& Ar ) override;
+	NETWORKFILE_API virtual bool Exec_Runtime(class UWorld* InWorld, const TCHAR* Cmd, FOutputDevice& Ar ) override;
 
 
 	/**
 	 * Send a heartbeat message to the file server. This will tell it we are alive, as well as
 	 * get back a list of files that have been updated on the server (so we can toss our copy)
 	 */
-	virtual void PerformHeartbeat();
+	NETWORKFILE_API virtual void PerformHeartbeat();
 
-	virtual void GetFileInfo(const TCHAR* Filename, FFileInfo& Info);
+	NETWORKFILE_API virtual void GetFileInfo(const TCHAR* Filename, FFileInfo& Info);
 
 	/**
 	 *	Convert the given filename from the server to the client version of it
@@ -197,12 +197,12 @@ protected:
 	 *
 	 *	@param	FilenameToConvert		Upon input, the server version of the filename. After the call, the client version
 	 */
-	virtual void ConvertServerFilenameToClientFilename(FString& FilenameToConvert);
+	NETWORKFILE_API virtual void ConvertServerFilenameToClientFilename(FString& FilenameToConvert);
 
-	virtual void FillGetFileList(FNetworkFileArchive& Payload);
+	NETWORKFILE_API virtual void FillGetFileList(FNetworkFileArchive& Payload);
 
-	virtual void ProcessServerInitialResponse(FArrayReader& InResponse, FPackageFileVersion& OutServerPackageVersion, int32& OutServerPackageLicenseeVersion);
-	virtual void ProcessServerCachedFilesResponse(FArrayReader& InReponse, const FPackageFileVersion& ServerPackageVersion, const int32 ServerPackageLicenseeVersion );
+	NETWORKFILE_API virtual void ProcessServerInitialResponse(FArrayReader& InResponse, FPackageFileVersion& OutServerPackageVersion, int32& OutServerPackageLicenseeVersion);
+	NETWORKFILE_API virtual void ProcessServerCachedFilesResponse(FArrayReader& InReponse, const FPackageFileVersion& ServerPackageVersion, const int32 ServerPackageLicenseeVersion );
 
 private:
 
@@ -213,7 +213,7 @@ private:
 	 * @param	Extension to test.
 	 * @return	True if Ext is a video extension.  e.g. .mp4
 	 */
-	static bool IsMediaExtension(const TCHAR* Ext);
+	static NETWORKFILE_API bool IsMediaExtension(const TCHAR* Ext);
 
 	/**
 	 * Returns whether the passed in extension is a an additional (but non asset) cooked file.
@@ -221,32 +221,32 @@ private:
 	 * @param	Extension to test.
 	 * @return	True if Ext is a additional cooked file.  e.g. .ubulk, .ufont
 	 */
-	static bool IsAdditionalCookedFileExtension(const TCHAR* Ext);
+	static NETWORKFILE_API bool IsAdditionalCookedFileExtension(const TCHAR* Ext);
 
 	/**
 	 * @return true if the path exists in a directory that should always use the local filesystem
 	 * This version does not worry about initialization or thread safety, do not call directly
 	 */
-	bool IsInLocalDirectoryUnGuarded(const FString& Filename);
+	NETWORKFILE_API bool IsInLocalDirectoryUnGuarded(const FString& Filename);
 
 	/**
 	 * @return true if the path exists in a directory that should always use the local filesystem
 	 */
-	bool IsInLocalDirectory(const FString& Filename);
+	NETWORKFILE_API bool IsInLocalDirectory(const FString& Filename);
 
 	/**
 	 * Given a filename, make sure the file exists on the local filesystem
 	 */
-	void EnsureFileIsLocal(const FString& Filename);
+	NETWORKFILE_API void EnsureFileIsLocal(const FString& Filename);
 
-	void OnCookOnTheFlyMessage(const UE::Cook::FCookOnTheFlyMessage& Message);
+	NETWORKFILE_API void OnCookOnTheFlyMessage(const UE::Cook::FCookOnTheFlyMessage& Message);
 
 protected:
 	/**
 	* Does normal path standardization, and also any extra modifications to make string comparisons against
 	* the internal directory list work properly.
 	*/
-	void MakeStandardNetworkFilename(FString& Filename);
+	NETWORKFILE_API void MakeStandardNetworkFilename(FString& Filename);
 
 protected:
 
@@ -312,10 +312,10 @@ private:
 	TSpscQueue<TArray<uint8>> PendingPayloads;
 	FEventRef NewPayloadEvent;
 
-	static FString MP4Extension;
-	static FString BulkFileExtension;
-	static FString ExpFileExtension;
-	static FString FontFileExtension;
+	static NETWORKFILE_API FString MP4Extension;
+	static NETWORKFILE_API FString BulkFileExtension;
+	static NETWORKFILE_API FString ExpFileExtension;
+	static NETWORKFILE_API FString FontFileExtension;
 };
 
 class SOCKETS_API FNetworkFileHandle : public IFileHandle
diff --git a/Engine/Source/Runtime/NetworkReplayStreaming/LocalFileNetworkReplayStreaming/Public/LocalFileNetworkReplayStreaming.h b/Engine/Source/Runtime/NetworkReplayStreaming/LocalFileNetworkReplayStreaming/Public/LocalFileNetworkReplayStreaming.h
--- a/Engine/Source/Runtime/NetworkReplayStreaming/LocalFileNetworkReplayStreaming/Public/LocalFileNetworkReplayStreaming.h
+++ b/Engine/Source/Runtime/NetworkReplayStreaming/LocalFileNetworkReplayStreaming/Public/LocalFileNetworkReplayStreaming.h
@@ -18,7 +18,7 @@
 class FNetworkReplayVersion;
 class FLocalFileNetworkReplayStreamer;
 
-struct LOCALFILENETWORKREPLAYSTREAMING_API FLocalFileReplayCustomVersion
+struct FLocalFileReplayCustomVersion
 {
 	enum Type
 	{
@@ -39,7 +39,7 @@ struct LOCALFILENETWORKREPLAYSTREAMING_API FLocalFileReplayCustomVersion
 	};
 
 	// The GUID for this custom version number
-	const static FGuid Guid;
+	LOCALFILENETWORKREPLAYSTREAMING_API const static FGuid Guid;
 
 	FLocalFileReplayCustomVersion() = delete;
 };
@@ -172,23 +172,23 @@ struct FLocalFileReplayInfo
 };
 
 /** Archive to wrap the file reader and respect chunk boundaries */
-class LOCALFILENETWORKREPLAYSTREAMING_API FLocalFileStreamFArchive : public FArchive
+class FLocalFileStreamFArchive : public FArchive
 {
 public:
 	PRAGMA_DISABLE_DEPRECATION_WARNINGS
 	FLocalFileStreamFArchive() : Pos(0), bAtEndOfReplay(false) {}
 	FLocalFileStreamFArchive(const FLocalFileStreamFArchive&) = default;
 	FLocalFileStreamFArchive& operator=(const FLocalFileStreamFArchive&) = default;
-	PRAGMA_ENABLE_DEPRECATION_WARNINGS
+	LOCALFILENETWORKREPLAYSTREAMING_API PRAGMA_ENABLE_DEPRECATION_WARNINGS
 
 	virtual void	Serialize(void* V, int64 Length) override;
-	virtual int64	Tell() override;
-	virtual int64	TotalSize() override;
-	virtual void	Seek(int64 InPos) override;
-	virtual bool	AtEnd() override;
+	LOCALFILENETWORKREPLAYSTREAMING_API virtual int64	Tell() override;
+	LOCALFILENETWORKREPLAYSTREAMING_API virtual int64	TotalSize() override;
+	LOCALFILENETWORKREPLAYSTREAMING_API virtual void	Seek(int64 InPos) override;
+	LOCALFILENETWORKREPLAYSTREAMING_API virtual bool	AtEnd() override;
 
-	int64 Tell() const;
-	int64 TotalSize() const;
+	LOCALFILENETWORKREPLAYSTREAMING_API int64 Tell() const;
+	LOCALFILENETWORKREPLAYSTREAMING_API int64 TotalSize() const;
 
 	void Reset()
 	{
@@ -505,62 +505,62 @@ protected:
 DECLARE_MULTICAST_DELEGATE_TwoParams(FOnLocalFileReplayFinishedWriting, const FString& /*StreamName*/, const FString& /*FullReplayFile*/);
 
 /** Local file streamer that supports playback/recording to a single file on disk */
-class LOCALFILENETWORKREPLAYSTREAMING_API FLocalFileNetworkReplayStreamer : public INetworkReplayStreamer, public TSharedFromThis<FLocalFileNetworkReplayStreamer>
+class FLocalFileNetworkReplayStreamer : public INetworkReplayStreamer, public TSharedFromThis<FLocalFileNetworkReplayStreamer>
 {
 	using FLocalFileReplayResult = UE::Net::TNetResult<ELocalFileReplayResult>;
 
 public:
-	FLocalFileNetworkReplayStreamer();
-	FLocalFileNetworkReplayStreamer(const FString& InDemoSavePath);
-	virtual ~FLocalFileNetworkReplayStreamer();
+	LOCALFILENETWORKREPLAYSTREAMING_API FLocalFileNetworkReplayStreamer();
+	LOCALFILENETWORKREPLAYSTREAMING_API FLocalFileNetworkReplayStreamer(const FString& InDemoSavePath);
+	LOCALFILENETWORKREPLAYSTREAMING_API virtual ~FLocalFileNetworkReplayStreamer();
 
 	/** INetworkReplayStreamer implementation */
-	virtual void StartStreaming(const FStartStreamingParameters& Params, const FStartStreamingCallback& Delegate) override;
-	virtual void StopStreaming() override;
-	virtual FArchive* GetHeaderArchive() override;
-	virtual FArchive* GetStreamingArchive() override;
-	virtual FArchive* GetCheckpointArchive() override;
-	virtual void FlushCheckpoint(const uint32 TimeInMS) override;
-	virtual void GotoCheckpointIndex(const int32 CheckpointIndex, const FGotoCallback& Delegate, EReplayCheckpointType CheckpointType) override;
-	virtual void GotoTimeInMS(const uint32 TimeInMS, const FGotoCallback& Delegate, EReplayCheckpointType CheckpointType) override;
-	virtual void UpdateTotalDemoTime(uint32 TimeInMS) override;
+	LOCALFILENETWORKREPLAYSTREAMING_API virtual void StartStreaming(const FStartStreamingParameters& Params, const FStartStreamingCallback& Delegate) override;
+	LOCALFILENETWORKREPLAYSTREAMING_API virtual void StopStreaming() override;
+	LOCALFILENETWORKREPLAYSTREAMING_API virtual FArchive* GetHeaderArchive() override;
+	LOCALFILENETWORKREPLAYSTREAMING_API virtual FArchive* GetStreamingArchive() override;
+	LOCALFILENETWORKREPLAYSTREAMING_API virtual FArchive* GetCheckpointArchive() override;
+	LOCALFILENETWORKREPLAYSTREAMING_API virtual void FlushCheckpoint(const uint32 TimeInMS) override;
+	LOCALFILENETWORKREPLAYSTREAMING_API virtual void GotoCheckpointIndex(const int32 CheckpointIndex, const FGotoCallback& Delegate, EReplayCheckpointType CheckpointType) override;
+	LOCALFILENETWORKREPLAYSTREAMING_API virtual void GotoTimeInMS(const uint32 TimeInMS, const FGotoCallback& Delegate, EReplayCheckpointType CheckpointType) override;
+	LOCALFILENETWORKREPLAYSTREAMING_API virtual void UpdateTotalDemoTime(uint32 TimeInMS) override;
 	virtual void UpdatePlaybackTime(uint32 TimeInMS) override {}
 	virtual uint32 GetTotalDemoTime() const override { return CurrentReplayInfo.LengthInMS; }
-	virtual bool IsDataAvailable() const override;
-	virtual void SetHighPriorityTimeRange(const uint32 StartTimeInMS, const uint32 EndTimeInMS) override;
-	virtual bool IsDataAvailableForTimeRange(const uint32 StartTimeInMS, const uint32 EndTimeInMS) override;
-	virtual bool IsLoadingCheckpoint() const override;
-	virtual bool IsLive() const override;
-	virtual void DeleteFinishedStream(const FString& StreamName, const FDeleteFinishedStreamCallback& Delegate) override;
-	virtual void DeleteFinishedStream( const FString& StreamName, const int32 UserIndex, const FDeleteFinishedStreamCallback& Delegate ) override;
-	virtual void EnumerateStreams( const FNetworkReplayVersion& InReplayVersion, const int32 UserIndex, const FString& MetaString, const TArray< FString >& ExtraParms, const FEnumerateStreamsCallback& Delegate ) override;
-	virtual void EnumerateRecentStreams( const FNetworkReplayVersion& ReplayVersion, const int32 UserIndex, const FEnumerateStreamsCallback& Delegate ) override;
-	virtual void AddUserToReplay(const FString& UserString) override;
-	virtual void AddEvent(const uint32 TimeInMS, const FString& Group, const FString& Meta, const TArray<uint8>& Data) override;
-	virtual void AddOrUpdateEvent(const FString& Name, const uint32 TimeInMS, const FString& Group, const FString& Meta, const TArray<uint8>& Data) override;
-	virtual void EnumerateEvents(const FString& Group, const FEnumerateEventsCallback& Delegate) override;
-	virtual void EnumerateEvents(const FString& ReplayName, const FString& Group, const FEnumerateEventsCallback& Delegate) override;
-	virtual void EnumerateEvents( const FString& ReplayName, const FString& Group, const int32 UserIndex, const FEnumerateEventsCallback& Delegate ) override;
-	virtual void RequestEventData(const FString& EventID, const FRequestEventDataCallback& Delegate) override;
-	virtual void RequestEventData(const FString& ReplayName, const FString& EventID, const FRequestEventDataCallback& Delegate) override;
-	virtual void RequestEventData(const FString& ReplayName, const FString& EventId, const int32 UserIndex, const FRequestEventDataCallback& Delegate) override;
-	virtual void RequestEventGroupData(const FString& Group, const FRequestEventGroupDataCallback& Delegate) override;
-	virtual void RequestEventGroupData(const FString& ReplayName, const FString& Group, const FRequestEventGroupDataCallback& Delegate) override;
-	virtual void RequestEventGroupData(const FString& ReplayName, const FString& Group, const int32 UserIndex, const FRequestEventGroupDataCallback& Delegate) override;
-	virtual void SearchEvents(const FString& EventGroup, const FSearchEventsCallback& Delegate) override;
-	virtual void KeepReplay(const FString& ReplayName, const bool bKeep, const FKeepReplayCallback& Delegate) override;
-	virtual void KeepReplay(const FString& ReplayName, const bool bKeep, const int32 UserIndex, const FKeepReplayCallback& Delegate) override;
-	virtual void RenameReplayFriendlyName(const FString& ReplayName, const FString& NewFriendlyName, const FRenameReplayCallback& Delegate) override;
-	virtual void RenameReplayFriendlyName(const FString& ReplayName, const FString& NewFriendlyName, const int32 UserIndex, const FRenameReplayCallback& Delegate) override;
-	virtual void RenameReplay(const FString& ReplayName, const FString& NewName, const FRenameReplayCallback& Delegate) override;
-	virtual void RenameReplay(const FString& ReplayName, const FString& NewName, const int32 UserIndex, const FRenameReplayCallback& Delegate) override;
+	LOCALFILENETWORKREPLAYSTREAMING_API virtual bool IsDataAvailable() const override;
+	LOCALFILENETWORKREPLAYSTREAMING_API virtual void SetHighPriorityTimeRange(const uint32 StartTimeInMS, const uint32 EndTimeInMS) override;
+	LOCALFILENETWORKREPLAYSTREAMING_API virtual bool IsDataAvailableForTimeRange(const uint32 StartTimeInMS, const uint32 EndTimeInMS) override;
+	LOCALFILENETWORKREPLAYSTREAMING_API virtual bool IsLoadingCheckpoint() const override;
+	LOCALFILENETWORKREPLAYSTREAMING_API virtual bool IsLive() const override;
+	LOCALFILENETWORKREPLAYSTREAMING_API virtual void DeleteFinishedStream(const FString& StreamName, const FDeleteFinishedStreamCallback& Delegate) override;
+	LOCALFILENETWORKREPLAYSTREAMING_API virtual void DeleteFinishedStream( const FString& StreamName, const int32 UserIndex, const FDeleteFinishedStreamCallback& Delegate ) override;
+	LOCALFILENETWORKREPLAYSTREAMING_API virtual void EnumerateStreams( const FNetworkReplayVersion& InReplayVersion, const int32 UserIndex, const FString& MetaString, const TArray< FString >& ExtraParms, const FEnumerateStreamsCallback& Delegate ) override;
+	LOCALFILENETWORKREPLAYSTREAMING_API virtual void EnumerateRecentStreams( const FNetworkReplayVersion& ReplayVersion, const int32 UserIndex, const FEnumerateStreamsCallback& Delegate ) override;
+	LOCALFILENETWORKREPLAYSTREAMING_API virtual void AddUserToReplay(const FString& UserString) override;
+	LOCALFILENETWORKREPLAYSTREAMING_API virtual void AddEvent(const uint32 TimeInMS, const FString& Group, const FString& Meta, const TArray<uint8>& Data) override;
+	LOCALFILENETWORKREPLAYSTREAMING_API virtual void AddOrUpdateEvent(const FString& Name, const uint32 TimeInMS, const FString& Group, const FString& Meta, const TArray<uint8>& Data) override;
+	LOCALFILENETWORKREPLAYSTREAMING_API virtual void EnumerateEvents(const FString& Group, const FEnumerateEventsCallback& Delegate) override;
+	LOCALFILENETWORKREPLAYSTREAMING_API virtual void EnumerateEvents(const FString& ReplayName, const FString& Group, const FEnumerateEventsCallback& Delegate) override;
+	LOCALFILENETWORKREPLAYSTREAMING_API virtual void EnumerateEvents( const FString& ReplayName, const FString& Group, const int32 UserIndex, const FEnumerateEventsCallback& Delegate ) override;
+	LOCALFILENETWORKREPLAYSTREAMING_API virtual void RequestEventData(const FString& EventID, const FRequestEventDataCallback& Delegate) override;
+	LOCALFILENETWORKREPLAYSTREAMING_API virtual void RequestEventData(const FString& ReplayName, const FString& EventID, const FRequestEventDataCallback& Delegate) override;
+	LOCALFILENETWORKREPLAYSTREAMING_API virtual void RequestEventData(const FString& ReplayName, const FString& EventId, const int32 UserIndex, const FRequestEventDataCallback& Delegate) override;
+	LOCALFILENETWORKREPLAYSTREAMING_API virtual void RequestEventGroupData(const FString& Group, const FRequestEventGroupDataCallback& Delegate) override;
+	LOCALFILENETWORKREPLAYSTREAMING_API virtual void RequestEventGroupData(const FString& ReplayName, const FString& Group, const FRequestEventGroupDataCallback& Delegate) override;
+	LOCALFILENETWORKREPLAYSTREAMING_API virtual void RequestEventGroupData(const FString& ReplayName, const FString& Group, const int32 UserIndex, const FRequestEventGroupDataCallback& Delegate) override;
+	LOCALFILENETWORKREPLAYSTREAMING_API virtual void SearchEvents(const FString& EventGroup, const FSearchEventsCallback& Delegate) override;
+	LOCALFILENETWORKREPLAYSTREAMING_API virtual void KeepReplay(const FString& ReplayName, const bool bKeep, const FKeepReplayCallback& Delegate) override;
+	LOCALFILENETWORKREPLAYSTREAMING_API virtual void KeepReplay(const FString& ReplayName, const bool bKeep, const int32 UserIndex, const FKeepReplayCallback& Delegate) override;
+	LOCALFILENETWORKREPLAYSTREAMING_API virtual void RenameReplayFriendlyName(const FString& ReplayName, const FString& NewFriendlyName, const FRenameReplayCallback& Delegate) override;
+	LOCALFILENETWORKREPLAYSTREAMING_API virtual void RenameReplayFriendlyName(const FString& ReplayName, const FString& NewFriendlyName, const int32 UserIndex, const FRenameReplayCallback& Delegate) override;
+	LOCALFILENETWORKREPLAYSTREAMING_API virtual void RenameReplay(const FString& ReplayName, const FString& NewName, const FRenameReplayCallback& Delegate) override;
+	LOCALFILENETWORKREPLAYSTREAMING_API virtual void RenameReplay(const FString& ReplayName, const FString& NewName, const int32 UserIndex, const FRenameReplayCallback& Delegate) override;
 	virtual FString	GetReplayID() const override { return CurrentStreamName; }
 	virtual EReplayStreamerState GetReplayStreamerState() const override { return StreamerState; }
 	virtual void SetTimeBufferHintSeconds(const float InTimeBufferHintSeconds) override {}
-	virtual void RefreshHeader() override;
-	virtual void DownloadHeader(const FDownloadHeaderCallback& Delegate) override;
+	LOCALFILENETWORKREPLAYSTREAMING_API virtual void RefreshHeader() override;
+	LOCALFILENETWORKREPLAYSTREAMING_API virtual void DownloadHeader(const FDownloadHeaderCallback& Delegate) override;
 
-	virtual bool IsCheckpointTypeSupported(EReplayCheckpointType CheckpointType) const override;
+	LOCALFILENETWORKREPLAYSTREAMING_API virtual bool IsCheckpointTypeSupported(EReplayCheckpointType CheckpointType) const override;
 
 	virtual bool SupportsCompression() const { return false; }
 
@@ -572,11 +572,11 @@ public:
 	virtual bool EncryptBuffer(TArrayView<const uint8> Plaintext, TArray<uint8>& Ciphertext, TArrayView<const uint8> EncryptionKey) const { return false; }
 	virtual bool DecryptBuffer(TArrayView<const uint8> Ciphertext, TArray<uint8>& Plaintext, TArrayView<const uint8> EncryptionKey) const { return false; }
 
-	bool AllowEncryptedWrite() const;
+	LOCALFILENETWORKREPLAYSTREAMING_API bool AllowEncryptedWrite() const;
 
-	void Tick(float DeltaSeconds);
+	LOCALFILENETWORKREPLAYSTREAMING_API void Tick(float DeltaSeconds);
 
-	virtual uint32 GetMaxFriendlyNameSize() const override;
+	LOCALFILENETWORKREPLAYSTREAMING_API virtual uint32 GetMaxFriendlyNameSize() const override;
 
 	virtual EStreamingOperationResult SetDemoPath(const FString& DemoPath) override
 	{
@@ -597,11 +597,11 @@ public:
 		return EStreamingOperationResult::Success;
 	}
 
-	void OnFileRequestComplete(const TSharedPtr<FQueuedLocalFileRequest, ESPMode::ThreadSafe>& Request);
+	LOCALFILENETWORKREPLAYSTREAMING_API void OnFileRequestComplete(const TSharedPtr<FQueuedLocalFileRequest, ESPMode::ThreadSafe>& Request);
 
-	bool IsStreaming() const;
+	LOCALFILENETWORKREPLAYSTREAMING_API bool IsStreaming() const;
 
-	bool HasPendingFileRequests() const;
+	LOCALFILENETWORKREPLAYSTREAMING_API bool HasPendingFileRequests() const;
 
 	void AddSimpleRequestToQueue(EQueuedLocalFileRequestType::Type RequestType, TFunction<void()>&& InFunction, TFunction<void()>&& InCompletionCallback)
 	{
@@ -642,16 +642,16 @@ public:
 	/** Map of checkpoint index to cached value */
 	TMap<int32, TSharedPtr<FCachedFileRequest>> DeltaCheckpointCache;
 
-	static FOnLocalFileReplayFinishedWriting OnReplayFinishedWriting;
+	static LOCALFILENETWORKREPLAYSTREAMING_API FOnLocalFileReplayFinishedWriting OnReplayFinishedWriting;
 
 protected:
 
-	void DeleteFinishedStream_Internal(const FString& StreamName, const int32 UserIndex, const FDeleteFinishedStreamCallback& Delegate);
-	void EnumerateEvents_Internal(const FString& ReplayName, const FString& Group, const int32 UserIndex, const FEnumerateEventsCallback& Delegate);
-	void RequestEventData_Internal(const FString& ReplayName, const FString& EventId, const int32 UserIndex, const FRequestEventDataCallback& Delegate);
-	void KeepReplay_Internal(const FString& ReplayName, const bool bKeep, const int32 UserIndex, const FKeepReplayCallback& Delegate);
-	void RenameReplayFriendlyName_Internal(const FString& ReplayName, const FString& NewFriendlyName, const int32 UserIndex, const FRenameReplayCallback& Delegate);
-	void RenameReplay_Internal(const FString& ReplayName, const FString& NewName, const int32 UserIndex, const FRenameReplayCallback& Delegate);
+	LOCALFILENETWORKREPLAYSTREAMING_API void DeleteFinishedStream_Internal(const FString& StreamName, const int32 UserIndex, const FDeleteFinishedStreamCallback& Delegate);
+	LOCALFILENETWORKREPLAYSTREAMING_API void EnumerateEvents_Internal(const FString& ReplayName, const FString& Group, const int32 UserIndex, const FEnumerateEventsCallback& Delegate);
+	LOCALFILENETWORKREPLAYSTREAMING_API void RequestEventData_Internal(const FString& ReplayName, const FString& EventId, const int32 UserIndex, const FRequestEventDataCallback& Delegate);
+	LOCALFILENETWORKREPLAYSTREAMING_API void KeepReplay_Internal(const FString& ReplayName, const bool bKeep, const int32 UserIndex, const FKeepReplayCallback& Delegate);
+	LOCALFILENETWORKREPLAYSTREAMING_API void RenameReplayFriendlyName_Internal(const FString& ReplayName, const FString& NewFriendlyName, const int32 UserIndex, const FRenameReplayCallback& Delegate);
+	LOCALFILENETWORKREPLAYSTREAMING_API void RenameReplay_Internal(const FString& ReplayName, const FString& NewName, const int32 UserIndex, const FRenameReplayCallback& Delegate);
 
 	/** 
 	 * Currently playing or recording replay metadata
@@ -675,10 +675,10 @@ protected:
 	TArray<TSharedPtr<FQueuedLocalFileRequest, ESPMode::ThreadSafe>> QueuedRequests;
 	TSharedPtr<FQueuedLocalFileRequest, ESPMode::ThreadSafe> ActiveRequest;
 
-	bool ProcessNextFileRequest();
-	bool IsFileRequestInProgress() const;
-	bool IsFileRequestPendingOrInProgress(const EQueuedLocalFileRequestType::Type RequestType) const;
-	void CancelStreamingRequests();
+	LOCALFILENETWORKREPLAYSTREAMING_API bool ProcessNextFileRequest();
+	LOCALFILENETWORKREPLAYSTREAMING_API bool IsFileRequestInProgress() const;
+	LOCALFILENETWORKREPLAYSTREAMING_API bool IsFileRequestPendingOrInProgress(const EQueuedLocalFileRequestType::Type RequestType) const;
+	LOCALFILENETWORKREPLAYSTREAMING_API void CancelStreamingRequests();
 
 	PRAGMA_DISABLE_DEPRECATION_WARNINGS
 	UE_DEPRECATED(5.1, "No longer used")
@@ -686,15 +686,15 @@ protected:
 	{ 
 		SetLastError(ELocalFileReplayResult::Unknown); 
 	}
-	PRAGMA_ENABLE_DEPRECATION_WARNINGS
+	LOCALFILENETWORKREPLAYSTREAMING_API PRAGMA_ENABLE_DEPRECATION_WARNINGS
 
 	void SetLastError(FLocalFileReplayResult&& Result);
 
-	void ConditionallyFlushStream();
-	void ConditionallyLoadNextChunk();
-	void ConditionallyRefreshReplayInfo();
+	LOCALFILENETWORKREPLAYSTREAMING_API void ConditionallyFlushStream();
+	LOCALFILENETWORKREPLAYSTREAMING_API void ConditionallyLoadNextChunk();
+	LOCALFILENETWORKREPLAYSTREAMING_API void ConditionallyRefreshReplayInfo();
 
-	void FlushCheckpointInternal(const uint32 TimeInMS);
+	LOCALFILENETWORKREPLAYSTREAMING_API void FlushCheckpointInternal(const uint32 TimeInMS);
 
 	struct FLocalFileSerializationInfo
 	{
@@ -714,35 +714,35 @@ protected:
 		return ReadReplayInfo(StreamName, OutReplayInfo, EReadReplayInfoFlags::None);
 	}
 
-	bool ReadReplayInfo(const FString& StreamName, FLocalFileReplayInfo& OutReplayInfo, EReadReplayInfoFlags Flags) const;
-	bool ReadReplayInfo(FArchive& Archive, FLocalFileReplayInfo& OutReplayInfo, EReadReplayInfoFlags Flags) const;
-	bool ReadReplayInfo(FArchive& Archive, FLocalFileReplayInfo& OutReplayInfo, struct FLocalFileSerializationInfo& SerializationInfo, EReadReplayInfoFlags Flags) const;
+	LOCALFILENETWORKREPLAYSTREAMING_API bool ReadReplayInfo(const FString& StreamName, FLocalFileReplayInfo& OutReplayInfo, EReadReplayInfoFlags Flags) const;
+	LOCALFILENETWORKREPLAYSTREAMING_API bool ReadReplayInfo(FArchive& Archive, FLocalFileReplayInfo& OutReplayInfo, EReadReplayInfoFlags Flags) const;
+	LOCALFILENETWORKREPLAYSTREAMING_API bool ReadReplayInfo(FArchive& Archive, FLocalFileReplayInfo& OutReplayInfo, struct FLocalFileSerializationInfo& SerializationInfo, EReadReplayInfoFlags Flags) const;
 
-	bool WriteReplayInfo(const FString& StreamName, const FLocalFileReplayInfo& ReplayInfo);
-	bool WriteReplayInfo(FArchive& Archive, const FLocalFileReplayInfo& ReplayInfo);
-	bool WriteReplayInfo(FArchive& Archive, const FLocalFileReplayInfo& InReplayInfo, struct FLocalFileSerializationInfo& SerializationInfo);
+	LOCALFILENETWORKREPLAYSTREAMING_API bool WriteReplayInfo(const FString& StreamName, const FLocalFileReplayInfo& ReplayInfo);
+	LOCALFILENETWORKREPLAYSTREAMING_API bool WriteReplayInfo(FArchive& Archive, const FLocalFileReplayInfo& ReplayInfo);
+	LOCALFILENETWORKREPLAYSTREAMING_API bool WriteReplayInfo(FArchive& Archive, const FLocalFileReplayInfo& InReplayInfo, struct FLocalFileSerializationInfo& SerializationInfo);
 
-	void FixupFriendlyNameLength(const FString& UnfixedName, FString& FixedName) const;
+	LOCALFILENETWORKREPLAYSTREAMING_API void FixupFriendlyNameLength(const FString& UnfixedName, FString& FixedName) const;
 
-	bool IsNamedStreamLive(const FString& StreamName) const;
+	LOCALFILENETWORKREPLAYSTREAMING_API bool IsNamedStreamLive(const FString& StreamName) const;
 
-	void FlushStream(const uint32 TimeInMS);
+	LOCALFILENETWORKREPLAYSTREAMING_API void FlushStream(const uint32 TimeInMS);
 
-	void WriteHeader();
+	LOCALFILENETWORKREPLAYSTREAMING_API void WriteHeader();
 
-	virtual TSharedPtr<FArchive> CreateLocalFileReader(const FString& InFilename) const;
-	virtual TSharedPtr<FArchive> CreateLocalFileWriter(const FString& InFilename) const;
-	virtual TSharedPtr<FArchive> CreateLocalFileWriterForOverwrite(const FString& InFilename) const;
+	LOCALFILENETWORKREPLAYSTREAMING_API virtual TSharedPtr<FArchive> CreateLocalFileReader(const FString& InFilename) const;
+	LOCALFILENETWORKREPLAYSTREAMING_API virtual TSharedPtr<FArchive> CreateLocalFileWriter(const FString& InFilename) const;
+	LOCALFILENETWORKREPLAYSTREAMING_API virtual TSharedPtr<FArchive> CreateLocalFileWriterForOverwrite(const FString& InFilename) const;
 
-	FString GetDemoPath() const;
+	LOCALFILENETWORKREPLAYSTREAMING_API FString GetDemoPath() const;
 	// Must be relative to the base demo path
 	virtual TArrayView<const FString> GetAdditionalRelativeDemoPaths() const { return {}; }
-	FString GetDemoFullFilename(const FString& FileName) const;
+	LOCALFILENETWORKREPLAYSTREAMING_API FString GetDemoFullFilename(const FString& FileName) const;
 
 	// Returns a name formatted as "demoX", where X is between 1 and MAX_DEMOS, inclusive.
 	// Returns the first value that doesn't yet exist, or if they all exist, returns the oldest one
 	// (it will be overwritten).
-	FString GetAutomaticDemoName() const;
+	LOCALFILENETWORKREPLAYSTREAMING_API FString GetAutomaticDemoName() const;
 
 	/** Handle to the archive that will read/write the demo header */
 	FLocalFileStreamFArchive HeaderAr;
@@ -761,52 +761,52 @@ protected:
 
 	FString DemoSavePath;
 
-	void AddRequestToCache(int32 ChunkIndex, const TArray<uint8>& RequestData);
-	void AddRequestToCache(int32 ChunkIndex, TArray<uint8>&& RequestData);
-	void CleanupRequestCache();
+	LOCALFILENETWORKREPLAYSTREAMING_API void AddRequestToCache(int32 ChunkIndex, const TArray<uint8>& RequestData);
+	LOCALFILENETWORKREPLAYSTREAMING_API void AddRequestToCache(int32 ChunkIndex, TArray<uint8>&& RequestData);
+	LOCALFILENETWORKREPLAYSTREAMING_API void CleanupRequestCache();
 
 	bool bCacheFileReadsInMemory;
 	mutable TMap<FString, TArray<uint8>> FileContentsCache;
-	const TArray<uint8>& GetCachedFileContents(const FString& Filename) const;
+	LOCALFILENETWORKREPLAYSTREAMING_API const TArray<uint8>& GetCachedFileContents(const FString& Filename) const;
 
-	void UpdateCurrentReplayInfo(FLocalFileReplayInfo& ReplayInfo, EUpdateReplayInfoFlags UpdateFlags = EUpdateReplayInfoFlags::None);
+	LOCALFILENETWORKREPLAYSTREAMING_API void UpdateCurrentReplayInfo(FLocalFileReplayInfo& ReplayInfo, EUpdateReplayInfoFlags UpdateFlags = EUpdateReplayInfoFlags::None);
 
-	virtual int32 GetDecompressedSizeBackCompat(FArchive& InCompressed) const;
+	LOCALFILENETWORKREPLAYSTREAMING_API virtual int32 GetDecompressedSizeBackCompat(FArchive& InCompressed) const;
 
 public:
-	static const FString& GetDefaultDemoSavePath();
-	static FString GetDemoFullFilename(const FString& DemoPath, const FString& FileName);
-	static bool CleanUpOldReplays(const FString& DemoPath = GetDefaultDemoSavePath(), TArrayView<const FString> AdditionalRelativeDemoPaths = {});
-	static bool GetDemoFreeStorageSpace(uint64& DiskFreeSpace, const FString& DemoPath);
+	static LOCALFILENETWORKREPLAYSTREAMING_API const FString& GetDefaultDemoSavePath();
+	static LOCALFILENETWORKREPLAYSTREAMING_API FString GetDemoFullFilename(const FString& DemoPath, const FString& FileName);
+	static LOCALFILENETWORKREPLAYSTREAMING_API bool CleanUpOldReplays(const FString& DemoPath = GetDefaultDemoSavePath(), TArrayView<const FString> AdditionalRelativeDemoPaths = {});
+	static LOCALFILENETWORKREPLAYSTREAMING_API bool GetDemoFreeStorageSpace(uint64& DiskFreeSpace, const FString& DemoPath);
 
-	static const uint32 FileMagic;
-	static const uint32 MaxFriendlyNameLen;
+	static LOCALFILENETWORKREPLAYSTREAMING_API const uint32 FileMagic;
+	static LOCALFILENETWORKREPLAYSTREAMING_API const uint32 MaxFriendlyNameLen;
 
 	UE_DEPRECATED(5.2, "No longer used, replaced with custom version.")
-	static const uint32 LatestVersion;
+	static LOCALFILENETWORKREPLAYSTREAMING_API const uint32 LatestVersion;
 
 private:
 	/** Manipulated by queued replay tasks, likely from another thread */
 	FLocalFileReplayInfo TaskReplayInfo;
 };
 
-class LOCALFILENETWORKREPLAYSTREAMING_API FLocalFileNetworkReplayStreamingFactory : public INetworkReplayStreamingFactory, public FTickableGameObject
+class FLocalFileNetworkReplayStreamingFactory : public INetworkReplayStreamingFactory, public FTickableGameObject
 {
 public:
-	virtual void StartupModule() override;
-	virtual void ShutdownModule() override;
+	LOCALFILENETWORKREPLAYSTREAMING_API virtual void StartupModule() override;
+	LOCALFILENETWORKREPLAYSTREAMING_API virtual void ShutdownModule() override;
 
-	virtual TSharedPtr<INetworkReplayStreamer> CreateReplayStreamer() override;
-	virtual void Flush() override;
+	LOCALFILENETWORKREPLAYSTREAMING_API virtual TSharedPtr<INetworkReplayStreamer> CreateReplayStreamer() override;
+	LOCALFILENETWORKREPLAYSTREAMING_API virtual void Flush() override;
 
 	/** FTickableGameObject */
-	virtual void Tick(float DeltaTime) override;
+	LOCALFILENETWORKREPLAYSTREAMING_API virtual void Tick(float DeltaTime) override;
 	virtual ETickableTickType GetTickableTickType() const override { return ETickableTickType::Always; }
-	virtual TStatId GetStatId() const override;
+	LOCALFILENETWORKREPLAYSTREAMING_API virtual TStatId GetStatId() const override;
 	bool IsTickableWhenPaused() const override { return true; }
 
 protected:
-	bool HasAnyPendingRequests() const;
+	LOCALFILENETWORKREPLAYSTREAMING_API bool HasAnyPendingRequests() const;
 
 	TArray<TSharedPtr<FLocalFileNetworkReplayStreamer>> LocalFileStreamers;
 };
diff --git a/Engine/Source/Runtime/NetworkReplayStreaming/NetworkReplayStreaming/Public/NetworkReplayStreaming.h b/Engine/Source/Runtime/NetworkReplayStreaming/NetworkReplayStreaming/Public/NetworkReplayStreaming.h
--- a/Engine/Source/Runtime/NetworkReplayStreaming/NetworkReplayStreaming/Public/NetworkReplayStreaming.h
+++ b/Engine/Source/Runtime/NetworkReplayStreaming/NetworkReplayStreaming/Public/NetworkReplayStreaming.h
@@ -512,7 +512,7 @@ DECLARE_MULTICAST_DELEGATE_TwoParams(FOnReplayGetAnalyticsAttributes, const clas
  * that delegate upon completion, and indicates success / failure through an appropriate
  * result type passed into the delegate.
  */
-class NETWORKREPLAYSTREAMING_API INetworkReplayStreamer 
+class INetworkReplayStreamer 
 {
 public:
 	virtual ~INetworkReplayStreamer() {}
@@ -625,14 +625,14 @@ public:
 	PRAGMA_DISABLE_DEPRECATION_WARNINGS
 	/** Returns the last error that occurred while streaming replays */
 	UE_DEPRECATED(5.1, "Deprecated in favor of HasError and HandleLastError")
-	virtual ENetworkReplayError::Type GetLastError() const;
-	PRAGMA_ENABLE_DEPRECATION_WARNINGS
+	NETWORKREPLAYSTREAMING_API virtual ENetworkReplayError::Type GetLastError() const;
+	NETWORKREPLAYSTREAMING_API PRAGMA_ENABLE_DEPRECATION_WARNINGS
 
 	virtual bool HasError() const;
 
-	virtual UE::Net::EHandleNetResult HandleLastError(UE::Net::FNetResultManager& ResultManager);
+	NETWORKREPLAYSTREAMING_API virtual UE::Net::EHandleNetResult HandleLastError(UE::Net::FNetResultManager& ResultManager);
 
-	virtual void SetExtendedError(UE::Net::FNetResult&& Result);
+	NETWORKREPLAYSTREAMING_API virtual void SetExtendedError(UE::Net::FNetResult&& Result);
 
 	/**
 	 * Adds a join-in-progress user to the set of users associated with the currently recording replay (if any)
@@ -667,12 +667,12 @@ public:
 
 	virtual void SetAnalyticsProvider(TSharedPtr<IAnalyticsProvider>& InProvider) {}
 
-	virtual TArray<FAnalyticsEventAttribute> AppendCommonReplayAttributes(TArray<FAnalyticsEventAttribute>&& Attrs) const;
+	NETWORKREPLAYSTREAMING_API virtual TArray<FAnalyticsEventAttribute> AppendCommonReplayAttributes(TArray<FAnalyticsEventAttribute>&& Attrs) const;
 
 	virtual void Exec(const TCHAR* Cmd, FOutputDevice& Ar) {}
 
 	/** Called from base streamer interface AppendCommonReplayAttributes to set common attributes */
-	static FOnReplayGetAnalyticsAttributes OnReplayGetAnalyticsAttributes;
+	static NETWORKREPLAYSTREAMING_API FOnReplayGetAnalyticsAttributes OnReplayGetAnalyticsAttributes;
 
 private:
 	TPimplPtr<UE::Net::FNetResult, EPimplPtrMode::DeepCopy> ExtendedError;
@@ -729,4 +729,4 @@ public:
 
 private:
 	TSet<FName> LoadedFactories;
-};
\ No newline at end of file
+};
diff --git a/Engine/Source/Runtime/NetworkReplayStreaming/SaveGameNetworkReplayStreaming/Public/SaveGameNetworkReplayStreaming.h b/Engine/Source/Runtime/NetworkReplayStreaming/SaveGameNetworkReplayStreaming/Public/SaveGameNetworkReplayStreaming.h
--- a/Engine/Source/Runtime/NetworkReplayStreaming/SaveGameNetworkReplayStreaming/Public/SaveGameNetworkReplayStreaming.h
+++ b/Engine/Source/Runtime/NetworkReplayStreaming/SaveGameNetworkReplayStreaming/Public/SaveGameNetworkReplayStreaming.h
@@ -29,53 +29,53 @@ enum class EGameDelegates_SaveGame : short;
  *
  * TODO: Proper handling of UserIndex.
  */
-class SAVEGAMENETWORKREPLAYSTREAMING_API FSaveGameNetworkReplayStreamer : public FLocalFileNetworkReplayStreamer
+class FSaveGameNetworkReplayStreamer : public FLocalFileNetworkReplayStreamer
 {
 public:
 
-	FSaveGameNetworkReplayStreamer();
-	FSaveGameNetworkReplayStreamer(const FString& DemoSavePath, const FString& PlaybackReplayName);
+	SAVEGAMENETWORKREPLAYSTREAMING_API FSaveGameNetworkReplayStreamer();
+	SAVEGAMENETWORKREPLAYSTREAMING_API FSaveGameNetworkReplayStreamer(const FString& DemoSavePath, const FString& PlaybackReplayName);
 
 	/** INetworkReplayStreamer implementation */
-	virtual void StartStreaming(const FStartStreamingParameters& Params, const FStartStreamingCallback& Delegate) override;
-	virtual void DeleteFinishedStream(const FString& ReplayName, const FDeleteFinishedStreamCallback& Delegate) override;
-	virtual void DeleteFinishedStream(const FString& StreamName, const int32 UserIndex, const FDeleteFinishedStreamCallback& Delegate) override;
-	virtual void EnumerateStreams(const FNetworkReplayVersion& InReplayVersion, const int32 UserIndex, const FString& MetaString, const TArray< FString >& ExtraParms, const FEnumerateStreamsCallback& Delegate) override;
-	virtual void EnumerateRecentStreams(const FNetworkReplayVersion& ReplayVersion, const int32 UserIndex, const FEnumerateStreamsCallback& Delegate) override;
-	virtual void KeepReplay(const FString& ReplayName, const bool bKeep, const FKeepReplayCallback& Delegate) override;
-	virtual void KeepReplay(const FString& ReplayName, const bool bKeep, const int32 UserIndex, const FKeepReplayCallback& Delegate) override;
-	virtual void RenameReplayFriendlyName(const FString& ReplayName, const FString& NewFriendlyName, const FRenameReplayCallback& Delegate) override;
-	virtual void RenameReplayFriendlyName(const FString& ReplayName, const FString& NewFriendlyName, const int32 UserIndex, const FRenameReplayCallback& Delegate) override;
-	virtual void RenameReplay(const FString& ReplayName, const FString& NewName, const FRenameReplayCallback& Delegate) override;
-	virtual void RenameReplay(const FString& ReplayName, const FString& NewName, const int32 UserIndex, const FRenameReplayCallback& Delegate) override;
-
-	virtual void EnumerateEvents(const FString& ReplayName, const FString& Group, const FEnumerateEventsCallback& Delegate) override;
-	virtual void EnumerateEvents( const FString& ReplayName, const FString& Group, const int32 UserIndex, const FEnumerateEventsCallback& Delegate) override;
-	virtual void RequestEventData(const FString& EventID, const FRequestEventDataCallback& Delegate) override;
-	virtual void RequestEventData(const FString& ReplayName, const FString& EventID, const FRequestEventDataCallback& Delegate) override;
-	virtual void RequestEventData(const FString& ReplayName, const FString& EventId, const int32 UserIndex, const FRequestEventDataCallback& Delegate) override;
-	virtual void RequestEventGroupData(const FString& Group, const FRequestEventGroupDataCallback& Delegate) override;
-	virtual void RequestEventGroupData(const FString& ReplayName, const FString& Group, const FRequestEventGroupDataCallback& Delegate) override;
-	virtual void RequestEventGroupData(const FString& ReplayName, const FString& Group, const int32 UserIndex, const FRequestEventGroupDataCallback& Delegate) override;
+	SAVEGAMENETWORKREPLAYSTREAMING_API virtual void StartStreaming(const FStartStreamingParameters& Params, const FStartStreamingCallback& Delegate) override;
+	SAVEGAMENETWORKREPLAYSTREAMING_API virtual void DeleteFinishedStream(const FString& ReplayName, const FDeleteFinishedStreamCallback& Delegate) override;
+	SAVEGAMENETWORKREPLAYSTREAMING_API virtual void DeleteFinishedStream(const FString& StreamName, const int32 UserIndex, const FDeleteFinishedStreamCallback& Delegate) override;
+	SAVEGAMENETWORKREPLAYSTREAMING_API virtual void EnumerateStreams(const FNetworkReplayVersion& InReplayVersion, const int32 UserIndex, const FString& MetaString, const TArray< FString >& ExtraParms, const FEnumerateStreamsCallback& Delegate) override;
+	SAVEGAMENETWORKREPLAYSTREAMING_API virtual void EnumerateRecentStreams(const FNetworkReplayVersion& ReplayVersion, const int32 UserIndex, const FEnumerateStreamsCallback& Delegate) override;
+	SAVEGAMENETWORKREPLAYSTREAMING_API virtual void KeepReplay(const FString& ReplayName, const bool bKeep, const FKeepReplayCallback& Delegate) override;
+	SAVEGAMENETWORKREPLAYSTREAMING_API virtual void KeepReplay(const FString& ReplayName, const bool bKeep, const int32 UserIndex, const FKeepReplayCallback& Delegate) override;
+	SAVEGAMENETWORKREPLAYSTREAMING_API virtual void RenameReplayFriendlyName(const FString& ReplayName, const FString& NewFriendlyName, const FRenameReplayCallback& Delegate) override;
+	SAVEGAMENETWORKREPLAYSTREAMING_API virtual void RenameReplayFriendlyName(const FString& ReplayName, const FString& NewFriendlyName, const int32 UserIndex, const FRenameReplayCallback& Delegate) override;
+	SAVEGAMENETWORKREPLAYSTREAMING_API virtual void RenameReplay(const FString& ReplayName, const FString& NewName, const FRenameReplayCallback& Delegate) override;
+	SAVEGAMENETWORKREPLAYSTREAMING_API virtual void RenameReplay(const FString& ReplayName, const FString& NewName, const int32 UserIndex, const FRenameReplayCallback& Delegate) override;
+
+	SAVEGAMENETWORKREPLAYSTREAMING_API virtual void EnumerateEvents(const FString& ReplayName, const FString& Group, const FEnumerateEventsCallback& Delegate) override;
+	SAVEGAMENETWORKREPLAYSTREAMING_API virtual void EnumerateEvents( const FString& ReplayName, const FString& Group, const int32 UserIndex, const FEnumerateEventsCallback& Delegate) override;
+	SAVEGAMENETWORKREPLAYSTREAMING_API virtual void RequestEventData(const FString& EventID, const FRequestEventDataCallback& Delegate) override;
+	SAVEGAMENETWORKREPLAYSTREAMING_API virtual void RequestEventData(const FString& ReplayName, const FString& EventID, const FRequestEventDataCallback& Delegate) override;
+	SAVEGAMENETWORKREPLAYSTREAMING_API virtual void RequestEventData(const FString& ReplayName, const FString& EventId, const int32 UserIndex, const FRequestEventDataCallback& Delegate) override;
+	SAVEGAMENETWORKREPLAYSTREAMING_API virtual void RequestEventGroupData(const FString& Group, const FRequestEventGroupDataCallback& Delegate) override;
+	SAVEGAMENETWORKREPLAYSTREAMING_API virtual void RequestEventGroupData(const FString& ReplayName, const FString& Group, const FRequestEventGroupDataCallback& Delegate) override;
+	SAVEGAMENETWORKREPLAYSTREAMING_API virtual void RequestEventGroupData(const FString& ReplayName, const FString& Group, const int32 UserIndex, const FRequestEventGroupDataCallback& Delegate) override;
 protected:
 
-	void StartStreamingSaved(const FStartStreamingParameters& Params, const FStartStreamingCallback& Delegate);
-	void DeleteFinishedStreamSaved(const FString& ReplayName, const int32 UserIndex, const FDeleteFinishedStreamCallback& Delegate) const;
-	void KeepReplaySaved(const FString& ReplayName, const bool bKeep, const int32 UserIndex, const FKeepReplayCallback& Result);
-	void RenameReplayFriendlyNameSaved(const FString& ReplayName, const FString& NewFriendlyName, const int32 UserIndex, const FRenameReplayCallback& Delegate) const;
-	void RenameReplaySaved(const FString& ReplayName, const FString& NewName, const int32 UserIndex, const FRenameReplayCallback& Delegate);
-	void EnumerateEventsSaved(const FString& ReplayName, const FString& Group, const int32 UserIndex, const FEnumerateEventsCallback& Delegate) const;
-	void RequestEventDataSaved(const FString& ReplayName, const FString& EventID, const int32 UserIndex, const FRequestEventDataCallback& Delegate);
-	void RequestEventGroupDataSaved(const FString& ReplayName, const FString& Group, const int32 UserIndex, const FRequestEventGroupDataCallback& Delegate);
-
-	void StartStreaming_Internal(const FStartStreamingParameters& Params, FStartStreamingResult& Result);
-	void DeleteFinishedStream_Internal(const FString& ReplayName, const int32 UserIndex, FDeleteFinishedStreamResult& Result) const;
-	void EnumerateStreams_Internal(const FNetworkReplayVersion& ReplayVersion, const int32 UserIndex, const FString& MetaString, const TArray< FString >& ExtraParms, FEnumerateStreamsResult& Result);
-	void KeepReplay_Internal(const FString& ReplayName, const bool bKeep, const int32 UserIndex, FKeepReplayResult& Result);
-	void RenameReplayFriendlyName_Internal(const FString& ReplayName, const FString& NewFriendlyName, const int32 UserIndex, FRenameReplayResult& Result) const;
-	void EnumerateEvents_Internal(const FString& ReplayName, const FString& Group, const int32 UserIndex, FEnumerateEventsResult& Result) const;
-	void RequestEventData_Internal(const FString& ReplayName, const FString& EventID, const int32 UserIndex, FRequestEventDataResult& Result);
-	void RequestEventGroupData_Internal(const FString& ReplayName, const FString& Group, const int32 UserIndex, FRequestEventGroupDataResult& Result);
+	SAVEGAMENETWORKREPLAYSTREAMING_API void StartStreamingSaved(const FStartStreamingParameters& Params, const FStartStreamingCallback& Delegate);
+	SAVEGAMENETWORKREPLAYSTREAMING_API void DeleteFinishedStreamSaved(const FString& ReplayName, const int32 UserIndex, const FDeleteFinishedStreamCallback& Delegate) const;
+	SAVEGAMENETWORKREPLAYSTREAMING_API void KeepReplaySaved(const FString& ReplayName, const bool bKeep, const int32 UserIndex, const FKeepReplayCallback& Result);
+	SAVEGAMENETWORKREPLAYSTREAMING_API void RenameReplayFriendlyNameSaved(const FString& ReplayName, const FString& NewFriendlyName, const int32 UserIndex, const FRenameReplayCallback& Delegate) const;
+	SAVEGAMENETWORKREPLAYSTREAMING_API void RenameReplaySaved(const FString& ReplayName, const FString& NewName, const int32 UserIndex, const FRenameReplayCallback& Delegate);
+	SAVEGAMENETWORKREPLAYSTREAMING_API void EnumerateEventsSaved(const FString& ReplayName, const FString& Group, const int32 UserIndex, const FEnumerateEventsCallback& Delegate) const;
+	SAVEGAMENETWORKREPLAYSTREAMING_API void RequestEventDataSaved(const FString& ReplayName, const FString& EventID, const int32 UserIndex, const FRequestEventDataCallback& Delegate);
+	SAVEGAMENETWORKREPLAYSTREAMING_API void RequestEventGroupDataSaved(const FString& ReplayName, const FString& Group, const int32 UserIndex, const FRequestEventGroupDataCallback& Delegate);
+
+	SAVEGAMENETWORKREPLAYSTREAMING_API void StartStreaming_Internal(const FStartStreamingParameters& Params, FStartStreamingResult& Result);
+	SAVEGAMENETWORKREPLAYSTREAMING_API void DeleteFinishedStream_Internal(const FString& ReplayName, const int32 UserIndex, FDeleteFinishedStreamResult& Result) const;
+	SAVEGAMENETWORKREPLAYSTREAMING_API void EnumerateStreams_Internal(const FNetworkReplayVersion& ReplayVersion, const int32 UserIndex, const FString& MetaString, const TArray< FString >& ExtraParms, FEnumerateStreamsResult& Result);
+	SAVEGAMENETWORKREPLAYSTREAMING_API void KeepReplay_Internal(const FString& ReplayName, const bool bKeep, const int32 UserIndex, FKeepReplayResult& Result);
+	SAVEGAMENETWORKREPLAYSTREAMING_API void RenameReplayFriendlyName_Internal(const FString& ReplayName, const FString& NewFriendlyName, const int32 UserIndex, FRenameReplayResult& Result) const;
+	SAVEGAMENETWORKREPLAYSTREAMING_API void EnumerateEvents_Internal(const FString& ReplayName, const FString& Group, const int32 UserIndex, FEnumerateEventsResult& Result) const;
+	SAVEGAMENETWORKREPLAYSTREAMING_API void RequestEventData_Internal(const FString& ReplayName, const FString& EventID, const int32 UserIndex, FRequestEventDataResult& Result);
+	SAVEGAMENETWORKREPLAYSTREAMING_API void RequestEventGroupData_Internal(const FString& ReplayName, const FString& Group, const int32 UserIndex, FRequestEventGroupDataResult& Result);
 
 	struct FSaveGameReplayVersionedInfo
 	{
@@ -105,25 +105,25 @@ protected:
 		int32 ReplayIndex;
 	};
 
-	bool StreamNameToSanitizedNames(const FString& StreamName, FSaveGameSanitizedNames& OutSanitizedNames) const;
-	void ReplayIndexToSanitizedNames(const int32 ReplayIndex, FSaveGameSanitizedNames& OutSanitizedNames) const;
+	SAVEGAMENETWORKREPLAYSTREAMING_API bool StreamNameToSanitizedNames(const FString& StreamName, FSaveGameSanitizedNames& OutSanitizedNames) const;
+	SAVEGAMENETWORKREPLAYSTREAMING_API void ReplayIndexToSanitizedNames(const int32 ReplayIndex, FSaveGameSanitizedNames& OutSanitizedNames) const;
 
-	bool ReadMetaDataFromLocalStream(FArchive& Archive, FSaveGameMetaData& OutMetaData) const;
-	bool ReadMetaDataFromSaveGame(class ISaveGameSystem& SaveGameSystem, const FSaveGameSanitizedNames& SanitizedNames, const int32 UserIndex, FSaveGameMetaData& OutMetaData, FStreamingResultBase& OutResult) const;
-	void PopulateStreamInfoFromMetaData(const FSaveGameMetaData& MetaData, FNetworkReplayStreamInfo& OutStreamInfo) const;
+	SAVEGAMENETWORKREPLAYSTREAMING_API bool ReadMetaDataFromLocalStream(FArchive& Archive, FSaveGameMetaData& OutMetaData) const;
+	SAVEGAMENETWORKREPLAYSTREAMING_API bool ReadMetaDataFromSaveGame(class ISaveGameSystem& SaveGameSystem, const FSaveGameSanitizedNames& SanitizedNames, const int32 UserIndex, FSaveGameMetaData& OutMetaData, FStreamingResultBase& OutResult) const;
+	SAVEGAMENETWORKREPLAYSTREAMING_API void PopulateStreamInfoFromMetaData(const FSaveGameMetaData& MetaData, FNetworkReplayStreamInfo& OutStreamInfo) const;
 
-	bool SerializeMetaData(FArchive& Archive, FSaveGameMetaData& MetaData) const;
-	bool SerializeVersionedMetaData(FArchive& Archive, FSaveGameMetaData& MetaData) const;
+	SAVEGAMENETWORKREPLAYSTREAMING_API bool SerializeMetaData(FArchive& Archive, FSaveGameMetaData& MetaData) const;
+	SAVEGAMENETWORKREPLAYSTREAMING_API bool SerializeVersionedMetaData(FArchive& Archive, FSaveGameMetaData& MetaData) const;
 
 	// Returns whether the input name corresponds to a save game.
-	bool IsSaveGameFileName(const FString& ReplayName) const;
+	SAVEGAMENETWORKREPLAYSTREAMING_API bool IsSaveGameFileName(const FString& ReplayName) const;
 
-	int32 GetReplayIndexFromName(const FString& ReplayName) const;
+	SAVEGAMENETWORKREPLAYSTREAMING_API int32 GetReplayIndexFromName(const FString& ReplayName) const;
 
-	FString GetFullPlaybackName() const;
-	FString GetLocalPlaybackName() const;
+	SAVEGAMENETWORKREPLAYSTREAMING_API FString GetFullPlaybackName() const;
+	SAVEGAMENETWORKREPLAYSTREAMING_API FString GetLocalPlaybackName() const;
 
-	virtual TArrayView<const FString> GetAdditionalRelativeDemoPaths() const override;
+	SAVEGAMENETWORKREPLAYSTREAMING_API virtual TArrayView<const FString> GetAdditionalRelativeDemoPaths() const override;
 
 	struct FSaveGameOptionInfo
 	{
@@ -152,13 +152,13 @@ protected:
 	// Special replay name that will be used when copying over SaveGame replays for playback.
 	const FString PlaybackReplayName;
 
-	static const FString& GetDefaultDemoSavePath();
-	static const FString& GetTempDemoRelativeSavePath();
-	static const FString& GetDefaultPlaybackName();
+	static SAVEGAMENETWORKREPLAYSTREAMING_API const FString& GetDefaultDemoSavePath();
+	static SAVEGAMENETWORKREPLAYSTREAMING_API const FString& GetTempDemoRelativeSavePath();
+	static SAVEGAMENETWORKREPLAYSTREAMING_API const FString& GetDefaultPlaybackName();
 
 private:
 
-	TFunction<bool(const TCHAR*, EGameDelegates_SaveGame, FString&)> WrapGetSaveGameOption() const;
+	SAVEGAMENETWORKREPLAYSTREAMING_API TFunction<bool(const TCHAR*, EGameDelegates_SaveGame, FString&)> WrapGetSaveGameOption() const;
 
 	// Although this isn't used on the GameThread, it should only be created / destroyed
 	// from the same thread. Therefore, no need to make it thread safe (for now).
@@ -167,10 +167,10 @@ private:
 	friend class FSaveGameNetworkReplayStreamingFactory;
 };
 
-class SAVEGAMENETWORKREPLAYSTREAMING_API FSaveGameNetworkReplayStreamingFactory : public FLocalFileNetworkReplayStreamingFactory
+class FSaveGameNetworkReplayStreamingFactory : public FLocalFileNetworkReplayStreamingFactory
 {
 public:
-	virtual void StartupModule() override;
+	SAVEGAMENETWORKREPLAYSTREAMING_API virtual void StartupModule() override;
 
-	virtual TSharedPtr<INetworkReplayStreamer> CreateReplayStreamer() override;
+	SAVEGAMENETWORKREPLAYSTREAMING_API virtual TSharedPtr<INetworkReplayStreamer> CreateReplayStreamer() override;
 };
diff --git a/Engine/Source/Runtime/Online/BackgroundHTTP/Public/BackgroundHttpManagerImpl.h b/Engine/Source/Runtime/Online/BackgroundHTTP/Public/BackgroundHttpManagerImpl.h
--- a/Engine/Source/Runtime/Online/BackgroundHTTP/Public/BackgroundHttpManagerImpl.h
+++ b/Engine/Source/Runtime/Online/BackgroundHTTP/Public/BackgroundHttpManagerImpl.h
@@ -21,57 +21,57 @@ class FBackgroundHttpFileHashHelper;
 /**
  * Contains implementation of some common functions that don't vary between implementation
  */
-class BACKGROUNDHTTP_API FBackgroundHttpManagerImpl 
+class FBackgroundHttpManagerImpl 
 	: public IBackgroundHttpManager
 	, public FTSTickerObjectBase
 {
 public:
-	FBackgroundHttpManagerImpl();
-	virtual ~FBackgroundHttpManagerImpl();
+	BACKGROUNDHTTP_API FBackgroundHttpManagerImpl();
+	BACKGROUNDHTTP_API virtual ~FBackgroundHttpManagerImpl();
 
-	virtual void AddRequest(const FBackgroundHttpRequestPtr Request) override;
-	virtual void RemoveRequest(const FBackgroundHttpRequestPtr Request) override;
+	BACKGROUNDHTTP_API virtual void AddRequest(const FBackgroundHttpRequestPtr Request) override;
+	BACKGROUNDHTTP_API virtual void RemoveRequest(const FBackgroundHttpRequestPtr Request) override;
 
-	virtual void Initialize() override;
-	virtual void Shutdown() override;
+	BACKGROUNDHTTP_API virtual void Initialize() override;
+	BACKGROUNDHTTP_API virtual void Shutdown() override;
 
-	virtual void DeleteAllTemporaryFiles() override;
+	BACKGROUNDHTTP_API virtual void DeleteAllTemporaryFiles() override;
 	
-	virtual int GetMaxActiveDownloads() const override;
-	virtual void SetMaxActiveDownloads(int MaxActiveDownloads) override;
+	BACKGROUNDHTTP_API virtual int GetMaxActiveDownloads() const override;
+	BACKGROUNDHTTP_API virtual void SetMaxActiveDownloads(int MaxActiveDownloads) override;
 
 	
-	virtual FString GetTempFileLocationForURL(const FString& URL) override;
+	BACKGROUNDHTTP_API virtual FString GetTempFileLocationForURL(const FString& URL) override;
 	
-	virtual void CleanUpDataAfterCompletingRequest(const FBackgroundHttpRequestPtr Request) override;
-	virtual void SetCellularPreference(int32 Value) override;
+	BACKGROUNDHTTP_API virtual void CleanUpDataAfterCompletingRequest(const FBackgroundHttpRequestPtr Request) override;
+	BACKGROUNDHTTP_API virtual void SetCellularPreference(int32 Value) override;
 	
 	//FTSTickerObjectBase implementation
-	virtual bool Tick(float DeltaTime) override;
+	BACKGROUNDHTTP_API virtual bool Tick(float DeltaTime) override;
 	
 protected:
-	virtual bool AssociateWithAnyExistingRequest(const FBackgroundHttpRequestPtr Request) override;
+	BACKGROUNDHTTP_API virtual bool AssociateWithAnyExistingRequest(const FBackgroundHttpRequestPtr Request) override;
 
-	virtual bool CheckForExistingCompletedDownload(const FBackgroundHttpRequestPtr Request, FString& ExistingFilePathOut, int64& ExistingFileSizeOut);
-	virtual void ActivatePendingRequests();
+	BACKGROUNDHTTP_API virtual bool CheckForExistingCompletedDownload(const FBackgroundHttpRequestPtr Request, FString& ExistingFilePathOut, int64& ExistingFileSizeOut);
+	BACKGROUNDHTTP_API virtual void ActivatePendingRequests();
 	
 	//Different from DeleteAllTemporaryFiles as this doesn't delete all files but rather cleans up specific bad files that have gone stale
-	virtual void DeleteStaleTempFiles();
+	BACKGROUNDHTTP_API virtual void DeleteStaleTempFiles();
 	
 	//Gets a list of full filenames for all temp files.
-	virtual void GatherAllTempFilenames(TArray<FString>& OutAllTempFilenames, bool bOutputAsFullPaths = false) const;
+	BACKGROUNDHTTP_API virtual void GatherAllTempFilenames(TArray<FString>& OutAllTempFilenames, bool bOutputAsFullPaths = false) const;
 	
 	//Helper function that converts the supplied list of Temp folder filenames to a list that is always all full paths.
-	virtual void ConvertAllTempFilenamesToFullPaths(TArray<FString>& OutFilenamesAsFullPaths, const TArray<FString>& FilenamesToConvertToFullPaths) const;
+	BACKGROUNDHTTP_API virtual void ConvertAllTempFilenamesToFullPaths(TArray<FString>& OutFilenamesAsFullPaths, const TArray<FString>& FilenamesToConvertToFullPaths) const;
 	
 	//Gather a list of any temp files that have timed out of our BackgroundHTTP settings.
 	//SecondsToConsiderOld should be a double representing how many seconds old a temp file needs to be to be returned by this function
 	//If OptionalFileListToCheck is empty will check all temp files in the backgroundhttp temp file folder. If supplied only the given file paths are checked
-	virtual void GatherTempFilesOlderThen(TArray<FString>& OutTimedOutTempFilenames,double SecondsToConsiderOld, TArray<FString>* OptionalFileList = nullptr) const;
+	BACKGROUNDHTTP_API virtual void GatherTempFilesOlderThen(TArray<FString>& OutTimedOutTempFilenames,double SecondsToConsiderOld, TArray<FString>* OptionalFileList = nullptr) const;
 	
 	//Gather a list of any temp files that have no corresponding URL Mapping entry
 	//If OptionalFileListToCheck is empty will check all temp files in the backgroundhttp temp file folder. If supplied only the given file paths are checked
-	virtual void GatherTempFilesWithoutURLMappings(TArray<FString>& OutTempFilesMissingURLMappings, TArray<FString>* OptionalFileList = nullptr) const;
+	BACKGROUNDHTTP_API virtual void GatherTempFilesWithoutURLMappings(TArray<FString>& OutTempFilesMissingURLMappings, TArray<FString>* OptionalFileList = nullptr) const;
 	
 	//Gets our FileHashHelper to compute temp file mappings
 	virtual BackgroundHttpFileHashHelperRef GetFileHashHelper(){ return FileHashHelper; }
diff --git a/Engine/Source/Runtime/Online/BackgroundHTTP/Public/BackgroundHttpNotificationObject.h b/Engine/Source/Runtime/Online/BackgroundHTTP/Public/BackgroundHttpNotificationObject.h
--- a/Engine/Source/Runtime/Online/BackgroundHTTP/Public/BackgroundHttpNotificationObject.h
+++ b/Engine/Source/Runtime/Online/BackgroundHTTP/Public/BackgroundHttpNotificationObject.h
@@ -9,15 +9,15 @@
 #include "Internationalization/Text.h"
 #include "Templates/SharedPointer.h"
 
-class BACKGROUNDHTTP_API FBackgroundHttpNotificationObject
+class FBackgroundHttpNotificationObject
 	: public TSharedFromThis<FBackgroundHttpNotificationObject, ESPMode::ThreadSafe>
 {
 public:
-	FBackgroundHttpNotificationObject(FText InNotificationTitle, FText InNotificationBody, FText InNotificationAction, const FString& InNotificationActivationString, bool InNotifyOnlyOnFullSuccess);
-	FBackgroundHttpNotificationObject(FText InNotificationTitle, FText InNotificationBody, FText InNotificationAction, const FString& InNotificationActivationString, bool InNotifyOnlyOnFullSuccess, bool bOnlySendNotificationInBackground, int32 InIdOverride);
-	~FBackgroundHttpNotificationObject();
+	BACKGROUNDHTTP_API FBackgroundHttpNotificationObject(FText InNotificationTitle, FText InNotificationBody, FText InNotificationAction, const FString& InNotificationActivationString, bool InNotifyOnlyOnFullSuccess);
+	BACKGROUNDHTTP_API FBackgroundHttpNotificationObject(FText InNotificationTitle, FText InNotificationBody, FText InNotificationAction, const FString& InNotificationActivationString, bool InNotifyOnlyOnFullSuccess, bool bOnlySendNotificationInBackground, int32 InIdOverride);
+	BACKGROUNDHTTP_API ~FBackgroundHttpNotificationObject();
 
-	void NotifyOfDownloadResult(bool bWasSuccess);
+	BACKGROUNDHTTP_API void NotifyOfDownloadResult(bool bWasSuccess);
 
 private:
 	void OnApp_EnteringBackground();
diff --git a/Engine/Source/Runtime/Online/BackgroundHTTP/Public/BackgroundHttpRequestImpl.h b/Engine/Source/Runtime/Online/BackgroundHTTP/Public/BackgroundHttpRequestImpl.h
--- a/Engine/Source/Runtime/Online/BackgroundHTTP/Public/BackgroundHttpRequestImpl.h
+++ b/Engine/Source/Runtime/Online/BackgroundHTTP/Public/BackgroundHttpRequestImpl.h
@@ -14,35 +14,35 @@ DECLARE_LOG_CATEGORY_EXTERN(LogBackgroundHttpRequest, Log, All)
 /**
  * Contains implementation of some common functions that don't have to vary between implementation
  */
-class BACKGROUNDHTTP_API FBackgroundHttpRequestImpl 
+class FBackgroundHttpRequestImpl 
 	: public IBackgroundHttpRequest
 {
 public:
-	FBackgroundHttpRequestImpl();
+	BACKGROUNDHTTP_API FBackgroundHttpRequestImpl();
 	virtual ~FBackgroundHttpRequestImpl() {}
 
 	//This should be called from the platform level when a BG download finishes.
-	virtual void OnBackgroundDownloadComplete();
+	BACKGROUNDHTTP_API virtual void OnBackgroundDownloadComplete();
 
 	// IBackgroundHttpRequest
-	virtual bool ProcessRequest() override;
-	virtual void CancelRequest() override;
-    virtual void PauseRequest() override;
-    virtual void ResumeRequest() override;
-    virtual void SetURLAsList(const TArray<FString>& URLs, int NumRetriesToAttempt) override;
-	virtual const TArray<FString>& GetURLList() const override;
-	virtual void SetCompleteNotification(FBackgroundHttpNotificationObjectPtr DownloadCompleteNotificationObjectIn) override;
-	virtual void CompleteWithExistingResponseData(FBackgroundHttpResponsePtr BackgroundResponse) override;
-	virtual FBackgroundHttpRequestCompleteDelegate& OnProcessRequestComplete() override;
-	virtual FBackgroundHttpProgressUpdateDelegate& OnProgressUpdated() override;
-	virtual const FBackgroundHttpResponsePtr GetResponse() const override;
-	virtual const FString& GetRequestID() const override;
-	virtual void SetRequestID(const FString& NewRequestID) override;
-	virtual bool HandleDelayedProcess() override;
-	virtual EBackgroundHTTPPriority GetRequestPriority() const override;
-	virtual void SetRequestPriority(EBackgroundHTTPPriority NewPriority) override;
+	BACKGROUNDHTTP_API virtual bool ProcessRequest() override;
+	BACKGROUNDHTTP_API virtual void CancelRequest() override;
+    BACKGROUNDHTTP_API virtual void PauseRequest() override;
+    BACKGROUNDHTTP_API virtual void ResumeRequest() override;
+    BACKGROUNDHTTP_API virtual void SetURLAsList(const TArray<FString>& URLs, int NumRetriesToAttempt) override;
+	BACKGROUNDHTTP_API virtual const TArray<FString>& GetURLList() const override;
+	BACKGROUNDHTTP_API virtual void SetCompleteNotification(FBackgroundHttpNotificationObjectPtr DownloadCompleteNotificationObjectIn) override;
+	BACKGROUNDHTTP_API virtual void CompleteWithExistingResponseData(FBackgroundHttpResponsePtr BackgroundResponse) override;
+	BACKGROUNDHTTP_API virtual FBackgroundHttpRequestCompleteDelegate& OnProcessRequestComplete() override;
+	BACKGROUNDHTTP_API virtual FBackgroundHttpProgressUpdateDelegate& OnProgressUpdated() override;
+	BACKGROUNDHTTP_API virtual const FBackgroundHttpResponsePtr GetResponse() const override;
+	BACKGROUNDHTTP_API virtual const FString& GetRequestID() const override;
+	BACKGROUNDHTTP_API virtual void SetRequestID(const FString& NewRequestID) override;
+	BACKGROUNDHTTP_API virtual bool HandleDelayedProcess() override;
+	BACKGROUNDHTTP_API virtual EBackgroundHTTPPriority GetRequestPriority() const override;
+	BACKGROUNDHTTP_API virtual void SetRequestPriority(EBackgroundHTTPPriority NewPriority) override;
 	
-	virtual void NotifyNotificationObjectOfComplete(bool bWasSuccess);
+	BACKGROUNDHTTP_API virtual void NotifyNotificationObjectOfComplete(bool bWasSuccess);
 protected:
 	TSharedPtr<FBackgroundHttpNotificationObject, ESPMode::ThreadSafe> DownloadCompleteNotificationObject;
 	FBackgroundHttpResponsePtr Response;
diff --git a/Engine/Source/Runtime/Online/BackgroundHTTP/Public/GenericPlatform/GenericPlatformBackgroundHttp.h b/Engine/Source/Runtime/Online/BackgroundHTTP/Public/GenericPlatform/GenericPlatformBackgroundHttp.h
--- a/Engine/Source/Runtime/Online/BackgroundHTTP/Public/GenericPlatform/GenericPlatformBackgroundHttp.h
+++ b/Engine/Source/Runtime/Online/BackgroundHTTP/Public/GenericPlatform/GenericPlatformBackgroundHttp.h
@@ -15,24 +15,24 @@ class IBackgroundHttpManager;
  * Intended usage is to use FPlatformBackgroundHttp instead of FGenericPlatformHttp
  * On platforms without a specific implementation, you should still use FPlatformBackgroundHttp and it will call into these functions
  */
-class BACKGROUNDHTTP_API FGenericPlatformBackgroundHttp
+class FGenericPlatformBackgroundHttp
 {
 public:
 	/**
 	 * Platform initialization step
 	 */
-	static void Initialize();
+	static BACKGROUNDHTTP_API void Initialize();
 
 	/**
 	 * Platform shutdown step
 	 */
-	static void Shutdown();
+	static BACKGROUNDHTTP_API void Shutdown();
 
 	/**
 	 * Creates a platform-specific Background HTTP manager.
 	 * Un-implemented platforms should create a FGenericPlatformBackgroundHttpManager
 	 */
-	static FBackgroundHttpManagerPtr CreatePlatformBackgroundHttpManager();
+	static BACKGROUNDHTTP_API FBackgroundHttpManagerPtr CreatePlatformBackgroundHttpManager();
 
 	/**
 	 * Creates a new Background Http request instance for the current platform
@@ -40,7 +40,7 @@ public:
 	 *
 	 * @return request object
 	 */
-	static FBackgroundHttpRequestPtr ConstructBackgroundRequest();
+	static BACKGROUNDHTTP_API FBackgroundHttpRequestPtr ConstructBackgroundRequest();
 
 	/**
 	 * Creates a new Background Http Response instance for the current platform
@@ -48,5 +48,5 @@ public:
 	 *
 	 * @return response object
 	 */
-	static FBackgroundHttpResponsePtr ConstructBackgroundResponse(int32 ResponseCode, const FString& TempFilePath);
+	static BACKGROUNDHTTP_API FBackgroundHttpResponsePtr ConstructBackgroundResponse(int32 ResponseCode, const FString& TempFilePath);
 };
diff --git a/Engine/Source/Runtime/Online/BackgroundHTTP/Public/GenericPlatform/GenericPlatformBackgroundHttpManager.h b/Engine/Source/Runtime/Online/BackgroundHTTP/Public/GenericPlatform/GenericPlatformBackgroundHttpManager.h
--- a/Engine/Source/Runtime/Online/BackgroundHTTP/Public/GenericPlatform/GenericPlatformBackgroundHttpManager.h
+++ b/Engine/Source/Runtime/Online/BackgroundHTTP/Public/GenericPlatform/GenericPlatformBackgroundHttpManager.h
@@ -10,10 +10,10 @@
 /**
  * Manages Background Http request that are currently being processed if no platform specific implementation has been made. 
  */
-class BACKGROUNDHTTP_API FGenericPlatformBackgroundHttpManager
+class FGenericPlatformBackgroundHttpManager
 	: public FBackgroundHttpManagerImpl
 {
 public:
 	virtual bool IsGenericImplementation() const override { return true; }
-	virtual ~FGenericPlatformBackgroundHttpManager();
-};
\ No newline at end of file
+	BACKGROUNDHTTP_API virtual ~FGenericPlatformBackgroundHttpManager();
+};
diff --git a/Engine/Source/Runtime/Online/BackgroundHTTP/Public/GenericPlatform/GenericPlatformBackgroundHttpRequest.h b/Engine/Source/Runtime/Online/BackgroundHTTP/Public/GenericPlatform/GenericPlatformBackgroundHttpRequest.h
--- a/Engine/Source/Runtime/Online/BackgroundHTTP/Public/GenericPlatform/GenericPlatformBackgroundHttpRequest.h
+++ b/Engine/Source/Runtime/Online/BackgroundHTTP/Public/GenericPlatform/GenericPlatformBackgroundHttpRequest.h
@@ -14,18 +14,18 @@
 /**
  * Contains implementation of some common functions that don't vary between implementation
  */
-class BACKGROUNDHTTP_API FGenericPlatformBackgroundHttpRequest 
+class FGenericPlatformBackgroundHttpRequest 
 	: public FBackgroundHttpRequestImpl
 {
 public:
 
-	FGenericPlatformBackgroundHttpRequest();
-	virtual ~FGenericPlatformBackgroundHttpRequest();
+	BACKGROUNDHTTP_API FGenericPlatformBackgroundHttpRequest();
+	BACKGROUNDHTTP_API virtual ~FGenericPlatformBackgroundHttpRequest();
 
 	//IHttpBackgroundHttpRequest
-	virtual bool HandleDelayedProcess() override;
-	virtual void CompleteWithExistingResponseData(FBackgroundHttpResponsePtr BackgroundResponse) override;
-	virtual void CancelRequest() override;
+	BACKGROUNDHTTP_API virtual bool HandleDelayedProcess() override;
+	BACKGROUNDHTTP_API virtual void CompleteWithExistingResponseData(FBackgroundHttpResponsePtr BackgroundResponse) override;
+	BACKGROUNDHTTP_API virtual void CancelRequest() override;
 
 protected:
 	//In the default implementation we actually use HTTPRequests instead of HttpBackgroundRequests.
diff --git a/Engine/Source/Runtime/Online/BackgroundHTTP/Public/Interfaces/IBackgroundHttpManager.h b/Engine/Source/Runtime/Online/BackgroundHTTP/Public/Interfaces/IBackgroundHttpManager.h
--- a/Engine/Source/Runtime/Online/BackgroundHTTP/Public/Interfaces/IBackgroundHttpManager.h
+++ b/Engine/Source/Runtime/Online/BackgroundHTTP/Public/Interfaces/IBackgroundHttpManager.h
@@ -6,7 +6,7 @@
 
 typedef TSharedPtr<class IBackgroundHttpManager, ESPMode::ThreadSafe> FBackgroundHttpManagerPtr;
 
-class BACKGROUNDHTTP_API IBackgroundHttpManager : public TSharedFromThis<IBackgroundHttpManager>
+class IBackgroundHttpManager : public TSharedFromThis<IBackgroundHttpManager>
 {
 public:
 
diff --git a/Engine/Source/Runtime/Online/BackgroundHTTP/Public/Interfaces/IBackgroundHttpModularFeature.h b/Engine/Source/Runtime/Online/BackgroundHTTP/Public/Interfaces/IBackgroundHttpModularFeature.h
--- a/Engine/Source/Runtime/Online/BackgroundHTTP/Public/Interfaces/IBackgroundHttpModularFeature.h
+++ b/Engine/Source/Runtime/Online/BackgroundHTTP/Public/Interfaces/IBackgroundHttpModularFeature.h
@@ -19,7 +19,7 @@ DECLARE_LOG_CATEGORY_EXTERN(LogBackgroundHttpModularFeature, Log, All)
  * This version of BackgroundHttp is designed to be used by any platform that relies on a ModularFeature to override the BackgroundHttp behavior. 
  * If no modular feature is found, we fall back on the GenericPlatform implementation to provide functionality
  */
-class BACKGROUNDHTTP_API IBackgroundHttpModularFeature : public IModularFeature
+class IBackgroundHttpModularFeature : public IModularFeature
 {
 public:
 	//Virtual methods for all PlatformBackgroundHttp required functions
diff --git a/Engine/Source/Runtime/Online/BackgroundHTTP/Public/Interfaces/IBackgroundHttpResponse.h b/Engine/Source/Runtime/Online/BackgroundHTTP/Public/Interfaces/IBackgroundHttpResponse.h
--- a/Engine/Source/Runtime/Online/BackgroundHTTP/Public/Interfaces/IBackgroundHttpResponse.h
+++ b/Engine/Source/Runtime/Online/BackgroundHTTP/Public/Interfaces/IBackgroundHttpResponse.h
@@ -7,7 +7,7 @@
 /**
  * Interface for Http Background Responses
  */
-class BACKGROUNDHTTP_API IBackgroundHttpResponse
+class IBackgroundHttpResponse
 {
 public:
 	/**
