- Add missing CoreOnline link dependencies
- This is preparation for making FUniqueNetIdWrapper::ToString call UE::Online::ToString which is not inlined, its dll exported.

diff --git a/Engine/Source/Runtime/NetworkReplayStreaming/HttpNetworkReplayStreaming/HttpNetworkReplayStreaming.Build.cs b/Engine/Source/Runtime/NetworkReplayStreaming/HttpNetworkReplayStreaming/HttpNetworkReplayStreaming.Build.cs
--- a/Engine/Source/Runtime/NetworkReplayStreaming/HttpNetworkReplayStreaming/HttpNetworkReplayStreaming.Build.cs
+++ b/Engine/Source/Runtime/NetworkReplayStreaming/HttpNetworkReplayStreaming/HttpNetworkReplayStreaming.Build.cs
@@ -13,6 +13,7 @@ namespace UnrealBuildTool.Rules
 				{
 					"Core",
 					"CoreUObject",
+					"CoreOnline",
                     "Engine",
                     "HTTP",
 					"NetworkReplayStreaming",
