Fix for #VSL-22-100 - Missing Signedness Checks in FHttpStreamFArchive::Seek() Can Lead to Out-of bounds Reads

diff --git a/Engine/Source/Runtime/NetworkReplayStreaming/HttpNetworkReplayStreaming/Private/HttpNetworkReplayStreaming.cpp b/Engine/Source/Runtime/NetworkReplayStreaming/HttpNetworkReplayStreaming/Private/HttpNetworkReplayStreaming.cpp
--- a/Engine/Source/Runtime/NetworkReplayStreaming/HttpNetworkReplayStreaming/Private/HttpNetworkReplayStreaming.cpp
+++ b/Engine/Source/Runtime/NetworkReplayStreaming/HttpNetworkReplayStreaming/Private/HttpNetworkReplayStreaming.cpp
@@ -199,9 +199,8 @@ int64 FHttpStreamFArchive::TotalSize()
 
 void FHttpStreamFArchive::Seek( int64 InPos ) 
 {
-	check( InPos <= Buffer.Num() );
-
-	Pos = IntCastChecked<int32>(InPos);
+	ensureMsgf( InPos >= 0 && InPos <= Buffer.Num(), TEXT("FHttpStreamFArchive::Seek invalid InPos %" INT64_FMT " Buffer Length %i"), InPos, Buffer.Num() );
+	Pos = FMath::Clamp( static_cast<int32>(InPos), 0, Buffer.Num() );
 }
 
 bool FHttpStreamFArchive::AtEnd() 
