LocalFileNetworkReplayStreaming file optimizations.
- Keep file writer open for duration of replay capture. This reduces the number of write operations which have a limit on Switch2.
- Don't seem if the requested seek position is the current position.
- Calculate file size using the available information so we don't need to flush the buffer.

#tests local

diff --git a/Engine/Source/Runtime/Core/Private/HAL/FileManagerGeneric.cpp b/Engine/Source/Runtime/Core/Private/HAL/FileManagerGeneric.cpp
--- a/Engine/Source/Runtime/Core/Private/HAL/FileManagerGeneric.cpp
+++ b/Engine/Source/Runtime/Core/Private/HAL/FileManagerGeneric.cpp
@@ -960,9 +960,7 @@ bool FArchiveFileWriterGeneric::SeekLowLevel( int64 InPos )
 
 int64 FArchiveFileWriterGeneric::TotalSize()
 {
-	// Make sure that all data is written before looking at file size.
-	FlushBuffer();
-	return Handle->Size();
+	return FMath::Max(Pos, Handle->Size());
 }
 
 bool FArchiveFileWriterGeneric::WriteLowLevel( const uint8* Src, int64 CountToWrite )
@@ -970,16 +968,19 @@ bool FArchiveFileWriterGeneric::WriteLowLevel( const uint8* Src, int64 CountToWr
 	return Handle->Write( Src, CountToWrite );
 }
 
-void FArchiveFileWriterGeneric::Seek( int64 InPos )
+void FArchiveFileWriterGeneric::Seek(int64 InPos)
 {
-	FlushBuffer();
-	if( !SeekLowLevel( InPos ) )
+	if (InPos != Pos)
 	{
-		uint32 LastError = FPlatformMisc::GetLastError();
-		SetError();
-		LogWriteError(TEXT("Error seeking file"), LastError);
+		FlushBuffer();
+		if (!SeekLowLevel(InPos))
+		{
+			uint32 LastError = FPlatformMisc::GetLastError();
+			SetError();
+			LogWriteError(TEXT("Error seeking file"), LastError);
+		}
+		Pos = InPos;
 	}
-	Pos = InPos;
 }
 
 bool FArchiveFileWriterGeneric::Close()
diff --git a/Engine/Source/Runtime/NetworkReplayStreaming/LocalFileNetworkReplayStreaming/Private/LocalFileNetworkReplayStreaming.cpp b/Engine/Source/Runtime/NetworkReplayStreaming/LocalFileNetworkReplayStreaming/Private/LocalFileNetworkReplayStreaming.cpp
--- a/Engine/Source/Runtime/NetworkReplayStreaming/LocalFileNetworkReplayStreaming/Private/LocalFileNetworkReplayStreaming.cpp
+++ b/Engine/Source/Runtime/NetworkReplayStreaming/LocalFileNetworkReplayStreaming/Private/LocalFileNetworkReplayStreaming.cpp
@@ -232,10 +232,10 @@ bool FLocalFileNetworkReplayStreamer::ReadReplayInfo(const FString& StreamName,
 {
 	SCOPE_CYCLE_COUNTER(STAT_LocalReplay_ReadReplayInfo);
 
-	TSharedPtr<FArchive> LocalFileAr = CreateLocalFileReader(GetDemoFullFilename(StreamName));
-	if (LocalFileAr.IsValid())
+	TSharedPtr<FArchive> LocalFileReader = CreateLocalFileReader(GetDemoFullFilename(StreamName));
+	if (LocalFileReader.IsValid())
 	{
-		return ReadReplayInfo(*LocalFileAr.Get(), Info, Flags);
+		return ReadReplayInfo(*LocalFileReader.Get(), Info, Flags);
 	}
 
 	return false;
@@ -630,20 +630,6 @@ bool FLocalFileNetworkReplayStreamer::ReadReplayInfo(FArchive& Archive, FLocalFi
 	return false;
 }
 
-bool FLocalFileNetworkReplayStreamer::WriteReplayInfo(const FString& StreamName, const FLocalFileReplayInfo& InReplayInfo)
-{
-	SCOPE_CYCLE_COUNTER(STAT_LocalReplay_WriteReplayInfo);
-
-	// Update metadata with latest info
-	TSharedPtr<FArchive> ReplayInfoFileAr = CreateLocalFileWriter(GetDemoFullFilename(StreamName));
-	if (ReplayInfoFileAr.IsValid())
-	{
-		return WriteReplayInfo(*ReplayInfoFileAr.Get(), InReplayInfo);
-	}
-
-	return false;
-}
-
 bool FLocalFileNetworkReplayStreamer::WriteReplayInfo(FArchive& Archive, const FLocalFileReplayInfo& InReplayInfo)
 {
 	FLocalFileSerializationInfo DefaultSerializationInfo;
@@ -901,6 +887,9 @@ void FLocalFileNetworkReplayStreamer::StartStreaming(const FStartStreamingParame
 			GenerateEncryptionKey(CurrentReplayInfo.EncryptionKey);
 		}
 
+		// LocalFileAr should be null if IsStreaming() returned false earlier */
+		check(LocalFileAr == nullptr);
+
 		AddDelegateFileRequestToQueue<FStartStreamingResult>(EQueuedLocalFileRequestType::StartRecording,
 			[this, Params, EncryptionKey=CurrentReplayInfo.EncryptionKey](TLocalFileRequestCommonData<FStartStreamingResult>& RequestData)
 			{
@@ -926,19 +915,30 @@ void FLocalFileNetworkReplayStreamer::StartStreaming(const FStartStreamingParame
 					return;
 				}
 
-				RequestData.DelegateResult.bRecording = true;
-
-				FLocalFileReplayInfo ExistingInfo;
-				if (ReadReplayInfo(FinalDemoName, ExistingInfo))
+				// Delete any existing demo with this name if it exists
+				if (IFileManager::Get().FileExists(*FullDemoFilename))
 				{
-					if (ExistingInfo.bIsLive)
+					FLocalFileReplayInfo ExistingInfo;
+					if (ReadReplayInfo(FinalDemoName, ExistingInfo))
 					{
-						UE_LOG(LogLocalFileReplay, Warning, TEXT("StartStreaming is overwriting an existing live replay file."));
+						if (ExistingInfo.bIsLive)
+						{
+							UE_LOG(LogLocalFileReplay, Warning, TEXT("StartStreaming is overwriting an existing live replay file."));
+						}
 					}
+					IFileManager::Get().Delete(*FullDemoFilename);
+				}
+
+				LocalFileAr = CreateLocalFileWriter(FullDemoFilename);
+				if (!LocalFileAr.IsValid())
+				{
+					UE_LOG(LogLocalFileReplay, Warning, TEXT("FLocalFileNetworkReplayStreamer::StartStreaming. LocalFileAr.IsValid == false."));
+					RequestData.AsyncError = ELocalFileReplayResult::FileWriter;
+					TaskReplayInfo.bIsValid = false;
+					return;
 				}
 
-				// Delete any existing demo with this name
-				IFileManager::Get().Delete(*FullDemoFilename);
+				RequestData.DelegateResult.bRecording = true;
 
 				TaskReplayInfo.NetworkVersion = Params.ReplayVersion.NetworkVersion;
 				TaskReplayInfo.Changelist = Params.ReplayVersion.Changelist;
@@ -947,7 +947,7 @@ void FLocalFileNetworkReplayStreamer::StartStreaming(const FStartStreamingParame
 				TaskReplayInfo.Timestamp = FDateTime::Now();
 				TaskReplayInfo.EncryptionKey = EncryptionKey;
 
-				TaskReplayInfo.bIsValid = WriteReplayInfo(CurrentStreamName, TaskReplayInfo);
+				TaskReplayInfo.bIsValid = WriteReplayInfo(*LocalFileAr.Get(), TaskReplayInfo);
 
 				if (!TaskReplayInfo.bIsValid)
 				{
@@ -1030,13 +1030,17 @@ void FLocalFileNetworkReplayStreamer::StopStreaming()
 
 		AddGenericRequestToQueue<ELocalFileReplayResult>(EQueuedLocalFileRequestType::StopRecording,
 			[this, TotalLengthInMS](ELocalFileReplayResult& ReplayResult)
+				{
+				if (LocalFileAr.IsValid())
 				{
 					// Set the final values of these header properties
-				TaskReplayInfo.bIsLive = false;
-				TaskReplayInfo.LengthInMS = TotalLengthInMS;
-				TaskReplayInfo.EncryptionKey = CurrentReplayInfo.EncryptionKey;
+					TaskReplayInfo.bIsLive = false;
+					TaskReplayInfo.LengthInMS = TotalLengthInMS;
+					TaskReplayInfo.EncryptionKey = CurrentReplayInfo.EncryptionKey;
 
-				WriteReplayInfo(CurrentStreamName, TaskReplayInfo);
+					WriteReplayInfo(*LocalFileAr.Get(), TaskReplayInfo);
+				}
+				LocalFileAr = nullptr;
 			},
 			[this](ELocalFileReplayResult& ReplayResult)
 			{
@@ -1273,7 +1277,6 @@ void FLocalFileNetworkReplayStreamer::AddOrUpdateEvent(const FString& Name, cons
 			// prefix with stream name to be consistent with http streamer
 			EventName = CurrentStreamName + TEXT("_") + EventName;
 
-			TSharedPtr<FArchive> LocalFileAr = CreateLocalFileWriter(GetDemoFullFilename(CurrentStreamName));
 			if (LocalFileAr.IsValid())
 			{
 				int32 EventIndex = INDEX_NONE;
@@ -1405,7 +1408,6 @@ void FLocalFileNetworkReplayStreamer::AddOrUpdateEvent(const FString& Name, cons
 				EventInfo.EventDataOffset = TypeOffset + InternalDataOffset;
 
 				LocalFileAr->Serialize(Writer.GetData(), Writer.TotalSize());
-				LocalFileAr = nullptr;
 			}
 			else
 			{
@@ -1528,8 +1530,8 @@ void FLocalFileNetworkReplayStreamer::RequestEventData_Internal(const FString& R
 				FLocalFileReplayInfo StoredReplayInfo;
 				if (ReadReplayInfo(FileName, StoredReplayInfo))
 				{
-					TSharedPtr<FArchive> LocalFileAr = CreateLocalFileReader(FullDemoFilename);
-					if (LocalFileAr.IsValid())
+					TSharedPtr<FArchive> LocalFileReader = CreateLocalFileReader(FullDemoFilename);
+					if (LocalFileReader.IsValid())
 					{
 						bool bEventFound = false;
 
@@ -1539,12 +1541,12 @@ void FLocalFileNetworkReplayStreamer::RequestEventData_Internal(const FString& R
 							{
 								bEventFound = true;
 
-								LocalFileAr->Seek(EventInfo.EventDataOffset);
+								LocalFileReader->Seek(EventInfo.EventDataOffset);
 
 								TArray<uint8> EventData;
 
 								EventData.AddUninitialized(EventInfo.SizeInBytes);
-								LocalFileAr->Serialize(EventData.GetData(), EventData.Num());
+								LocalFileReader->Serialize(EventData.GetData(), EventData.Num());
 
 								if (StoredReplayInfo.bEncrypted)
 								{
@@ -1578,7 +1580,7 @@ void FLocalFileNetworkReplayStreamer::RequestEventData_Internal(const FString& R
 							}
 						}
 
-						LocalFileAr = nullptr;
+						LocalFileReader = nullptr;
 
 						// we didn't find the event
 						if (!bEventFound)
@@ -1635,8 +1637,8 @@ void FLocalFileNetworkReplayStreamer::RequestEventGroupData(const FString& Repla
 				FLocalFileReplayInfo StoredReplayInfo;
 				if (ReadReplayInfo(FileName, StoredReplayInfo))
 				{
-					TSharedPtr<FArchive> LocalFileAr = CreateLocalFileReader(FullDemoFilename);
-					if (LocalFileAr.IsValid())
+					TSharedPtr<FArchive> LocalFileReader = CreateLocalFileReader(FullDemoFilename);
+					if (LocalFileReader.IsValid())
 					{
 						bool bGroupFound = false;
 
@@ -1646,7 +1648,7 @@ void FLocalFileNetworkReplayStreamer::RequestEventGroupData(const FString& Repla
 							{
 								bGroupFound = true;
 
-								LocalFileAr->Seek(EventInfo.EventDataOffset);
+								LocalFileReader->Seek(EventInfo.EventDataOffset);
 
 								RequestData.DelegateResult.Result = EStreamingOperationResult::Success;
 								
@@ -1660,7 +1662,7 @@ void FLocalFileNetworkReplayStreamer::RequestEventGroupData(const FString& Repla
 								TArray<uint8> EventData;
 
 								EventData.AddUninitialized(EventInfo.SizeInBytes);
-								LocalFileAr->Serialize(EventData.GetData(), EventData.Num());
+								LocalFileReader->Serialize(EventData.GetData(), EventData.Num());
 
 								if (StoredReplayInfo.bEncrypted)
 								{
@@ -1699,7 +1701,7 @@ void FLocalFileNetworkReplayStreamer::RequestEventGroupData(const FString& Repla
 							RequestData.DelegateResult.Result = EStreamingOperationResult::EventNotFound;
 						}
 
-						LocalFileAr = nullptr;
+						LocalFileReader = nullptr;
 					}
 					else
 					{
@@ -1783,8 +1785,8 @@ void FLocalFileNetworkReplayStreamer::RenameReplayFriendlyName_Internal(const FS
 
 			// Do this inside a scope, to make sure the file archive is closed before continuing.
 			{
-				TSharedPtr<FArchive> ReadAr = CreateLocalFileReader(FullReplayName);
-				if (!ReadAr.IsValid() || ReadAr->TotalSize() <= 0 || !ReadReplayInfo(*ReadAr.Get(), TempReplayInfo, SerializationInfo, EReadReplayInfoFlags::None))
+				TSharedPtr<FArchive> LocalFileReader = CreateLocalFileReader(FullReplayName);
+				if (!LocalFileReader.IsValid() || LocalFileReader->TotalSize() <= 0 || !ReadReplayInfo(*LocalFileReader.Get(), TempReplayInfo, SerializationInfo, EReadReplayInfoFlags::None))
 				{
 					UE_LOG(LogLocalFileReplay, Warning, TEXT("FLocalFileNetworkReplayStreamer::RenameReplayFriendlyName: Failed to read replay info %s"), *ReplayName);
 					return;
@@ -1801,8 +1803,11 @@ void FLocalFileNetworkReplayStreamer::RenameReplayFriendlyName_Internal(const FS
 
 			// Do this inside a scope, to make sure the file archive is closed before continuing.
 			{
-				TSharedPtr<FArchive> WriteAr = CreateLocalFileWriter(FullReplayName);
-				if (!WriteAr.IsValid() || WriteAr->TotalSize() <= 0 || !WriteReplayInfo(*WriteAr.Get(), TempReplayInfo, SerializationInfo))
+				// ensure existing replay is flushed and closed
+				LocalFileAr = nullptr;
+
+				TSharedPtr<FArchive> LocalFileWriter = CreateLocalFileWriter(FullReplayName);
+				if (!LocalFileWriter.IsValid() || LocalFileWriter->TotalSize() <= 0 || !WriteReplayInfo(*LocalFileWriter.Get(), TempReplayInfo, SerializationInfo))
 				{
 					UE_LOG(LogLocalFileReplay, Warning, TEXT("FLocalFileNetworkReplayStreamer::RenameReplayFriendlyName: Failed to write replay info %s"), *ReplayName);
 					return;
@@ -1899,7 +1904,6 @@ void FLocalFileNetworkReplayStreamer::FlushStream(const uint32 TimeInMS)
 		{
 			SCOPE_CYCLE_COUNTER(STAT_LocalReplay_FlushStream);
 
-			TSharedPtr<FArchive> LocalFileAr = CreateLocalFileWriter(GetDemoFullFilename(CurrentStreamName));
 			if (LocalFileAr.IsValid())
 			{
 				LocalFileAr->Seek(LocalFileAr->TotalSize());
@@ -1995,11 +1999,9 @@ void FLocalFileNetworkReplayStreamer::FlushStream(const uint32 TimeInMS)
 					TaskReplayInfo.TotalDataSizeInBytes += DataInfo.MemorySizeInBytes;
 				}
 
-				LocalFileAr = nullptr;
-
 				TaskReplayInfo.LengthInMS = TotalLengthInMS;
 
-				WriteReplayInfo(CurrentStreamName, TaskReplayInfo);
+				WriteReplayInfo(*LocalFileAr.Get(), TaskReplayInfo);
 			}
 			else
 			{
@@ -2065,7 +2067,6 @@ void FLocalFileNetworkReplayStreamer::FlushCheckpointInternal(const uint32 TimeI
 			const int32 DataChunkIndex = TaskReplayInfo.DataChunks.Num();
 			const int32 CheckpointIndex = TaskReplayInfo.Checkpoints.Num();
 
-			TSharedPtr<FArchive> LocalFileAr = CreateLocalFileWriter(GetDemoFullFilename(CurrentStreamName));
 			if (LocalFileAr.IsValid())
 			{
 				LocalFileAr->Seek(LocalFileAr->TotalSize());
@@ -2165,7 +2166,6 @@ void FLocalFileNetworkReplayStreamer::FlushCheckpointInternal(const uint32 TimeI
 					CheckpointInfo.EventDataOffset = EventDataOffset;
 				}
 
-				LocalFileAr = nullptr;
 			}
 			else
 			{
@@ -2175,7 +2175,7 @@ void FLocalFileNetworkReplayStreamer::FlushCheckpointInternal(const uint32 TimeI
 
 			TaskReplayInfo.LengthInMS = TotalLengthInMS;
 
-			WriteReplayInfo(CurrentStreamName, TaskReplayInfo);
+			WriteReplayInfo(*LocalFileAr.Get(), TaskReplayInfo);
 		},
 		[this](ELocalFileReplayResult& ReplayResult)
 		{
@@ -2275,8 +2275,8 @@ void FLocalFileNetworkReplayStreamer::GotoCheckpointIndex(const int32 Checkpoint
 
 			const FString FullDemoFilename = GetDemoFullFilename(CurrentStreamName);
 
-			TSharedPtr<FArchive> LocalFileAr = CreateLocalFileReader(FullDemoFilename);
-			if (LocalFileAr.IsValid())
+			TSharedPtr<FArchive> LocalFileReader = CreateLocalFileReader(FullDemoFilename);
+			if (LocalFileReader.IsValid())
 			{
 					TArray<uint8> CheckpointData;
 
@@ -2293,12 +2293,12 @@ void FLocalFileNetworkReplayStreamer::GotoCheckpointIndex(const int32 Checkpoint
 						}
 						else
 						{
-							LocalFileAr->Seek(TaskReplayInfo.Checkpoints[i].EventDataOffset);
+							LocalFileReader->Seek(TaskReplayInfo.Checkpoints[i].EventDataOffset);
 
 							CheckpointData.Reset();
 							CheckpointData.AddUninitialized(TaskReplayInfo.Checkpoints[i].SizeInBytes);
 
-							LocalFileAr->Serialize(CheckpointData.GetData(), CheckpointData.Num());
+							LocalFileReader->Serialize(CheckpointData.GetData(), CheckpointData.Num());
 
 							TArray<uint8> PlaintextData;
 
@@ -2477,14 +2477,14 @@ void FLocalFileNetworkReplayStreamer::GotoCheckpointIndex(const int32 Checkpoint
 
 			const FString FullDemoFilename = GetDemoFullFilename(CurrentStreamName);
 
-			TSharedPtr<FArchive> LocalFileAr = CreateLocalFileReader(FullDemoFilename);
-			if (LocalFileAr.IsValid())
+			TSharedPtr<FArchive> LocalFileReader = CreateLocalFileReader(FullDemoFilename);
+			if (LocalFileReader.IsValid())
 			{
-				LocalFileAr->Seek(TaskReplayInfo.Checkpoints[CheckpointIndex].EventDataOffset);
+				LocalFileReader->Seek(TaskReplayInfo.Checkpoints[CheckpointIndex].EventDataOffset);
 
 				RequestData.DataBuffer.AddUninitialized(TaskReplayInfo.Checkpoints[CheckpointIndex].SizeInBytes);
 
-					LocalFileAr->Serialize(RequestData.DataBuffer.GetData(), RequestData.DataBuffer.Num());
+				LocalFileReader->Serialize(RequestData.DataBuffer.GetData(), RequestData.DataBuffer.Num());
 
 					// Get the checkpoint data
 					
@@ -2538,7 +2538,7 @@ void FLocalFileNetworkReplayStreamer::GotoCheckpointIndex(const int32 Checkpoint
 						}
 					}
 
-				LocalFileAr = nullptr;
+				LocalFileReader = nullptr;
 			}
 		},
 		[this, CheckpointIndex, Delegate](TLocalFileRequestCommonData<FGotoResult>& RequestData)
@@ -3112,17 +3112,17 @@ void FLocalFileNetworkReplayStreamer::DownloadHeader(const FDownloadHeaderCallba
 
 				const FString FullDemoFilename = GetDemoFullFilename(CurrentStreamName);
 
-				TSharedPtr<FArchive> LocalFileAr = CreateLocalFileReader(FullDemoFilename);
-				if (LocalFileAr.IsValid())
+				TSharedPtr<FArchive> LocalFileReader = CreateLocalFileReader(FullDemoFilename);
+				if (LocalFileReader.IsValid())
 				{
 					HeaderData.AddUninitialized(HeaderSize);
 
-					LocalFileAr->Seek(HeaderOffset);
-					LocalFileAr->Serialize((void*)HeaderData.GetData(), HeaderData.Num());
+					LocalFileReader->Seek(HeaderOffset);
+					LocalFileReader->Serialize((void*)HeaderData.GetData(), HeaderData.Num());
 
 					RequestData.DelegateResult.Result = EStreamingOperationResult::Success;
 
-					LocalFileAr = nullptr;
+					LocalFileReader = nullptr;
 				}
 			},
 			[this, Delegate](TLocalFileRequestCommonData<FDownloadHeaderResult>& RequestData)
@@ -3168,7 +3168,6 @@ void FLocalFileNetworkReplayStreamer::WriteHeader()
 		{
 			SCOPE_CYCLE_COUNTER(STAT_LocalReplay_FlushHeader);
 
-			TSharedPtr<FArchive> LocalFileAr = CreateLocalFileWriter(GetDemoFullFilename(CurrentStreamName));
 			if (LocalFileAr.IsValid())
 			{
 				if (TaskReplayInfo.HeaderChunkIndex == INDEX_NONE)
@@ -3217,7 +3216,6 @@ void FLocalFileNetworkReplayStreamer::WriteHeader()
 				const int32 DataOffset = IntCastChecked<int32>(LocalFileAr->Tell());
 
 				LocalFileAr->Serialize((void*)HeaderData.GetData(), HeaderData.Num());
-				LocalFileAr = nullptr;
 
 				if (TaskReplayInfo.HeaderChunkIndex == INDEX_NONE)
 				{
@@ -3390,14 +3388,14 @@ void FLocalFileNetworkReplayStreamer::ConditionallyLoadNextChunk()
 			
 				const FString FullDemoFilename = GetDemoFullFilename(CurrentStreamName);
 		
-				TSharedPtr<FArchive> LocalFileAr = CreateLocalFileReader(FullDemoFilename);
-				if (LocalFileAr.IsValid())
+				TSharedPtr<FArchive> LocalFileReader = CreateLocalFileReader(FullDemoFilename);
+				if (LocalFileReader.IsValid())
 				{
-				LocalFileAr->Seek(TaskReplayInfo.DataChunks[RequestedStreamChunkIndex].ReplayDataOffset);
+					LocalFileReader->Seek(TaskReplayInfo.DataChunks[RequestedStreamChunkIndex].ReplayDataOffset);
 
 				RequestData.DataBuffer.AddUninitialized(TaskReplayInfo.DataChunks[RequestedStreamChunkIndex].SizeInBytes);
 
-					LocalFileAr->Serialize(RequestData.DataBuffer.GetData(), RequestData.DataBuffer.Num());
+				LocalFileReader->Serialize(RequestData.DataBuffer.GetData(), RequestData.DataBuffer.Num());
 
 				if (TaskReplayInfo.bEncrypted)
 					{
@@ -3455,7 +3453,7 @@ void FLocalFileNetworkReplayStreamer::ConditionallyLoadNextChunk()
 						}
 					}
 
-					LocalFileAr = nullptr;
+					LocalFileReader = nullptr;
 				}
 		},
 		[this, RequestedStreamChunkIndex](TLocalFileRequestCommonData<FStreamingResultBase>& RequestData)
diff --git a/Engine/Source/Runtime/NetworkReplayStreaming/LocalFileNetworkReplayStreaming/Public/LocalFileNetworkReplayStreaming.h b/Engine/Source/Runtime/NetworkReplayStreaming/LocalFileNetworkReplayStreaming/Public/LocalFileNetworkReplayStreaming.h
--- a/Engine/Source/Runtime/NetworkReplayStreaming/LocalFileNetworkReplayStreaming/Public/LocalFileNetworkReplayStreaming.h
+++ b/Engine/Source/Runtime/NetworkReplayStreaming/LocalFileNetworkReplayStreaming/Public/LocalFileNetworkReplayStreaming.h
@@ -636,6 +636,9 @@ public:
 		QueuedRequests.Add(MakeShared<TGenericCachedLocalFileRequest<DelegateResultType>, ESPMode::ThreadSafe>(InCacheKey, AsShared(), RequestType, MoveTemp(InFunction), MoveTemp(InCompletionCallback)));
 	}
 
+	/** Current stream Archive */
+	TSharedPtr<FArchive> LocalFileAr;
+
 	/** Map of chunk index to cached value */
 	TMap<int32, TSharedPtr<FCachedFileRequest>> RequestCache;
 
@@ -718,7 +721,6 @@ protected:
 	LOCALFILENETWORKREPLAYSTREAMING_API bool ReadReplayInfo(FArchive& Archive, FLocalFileReplayInfo& OutReplayInfo, EReadReplayInfoFlags Flags) const;
 	LOCALFILENETWORKREPLAYSTREAMING_API bool ReadReplayInfo(FArchive& Archive, FLocalFileReplayInfo& OutReplayInfo, struct FLocalFileSerializationInfo& SerializationInfo, EReadReplayInfoFlags Flags) const;
 
-	LOCALFILENETWORKREPLAYSTREAMING_API bool WriteReplayInfo(const FString& StreamName, const FLocalFileReplayInfo& ReplayInfo);
 	LOCALFILENETWORKREPLAYSTREAMING_API bool WriteReplayInfo(FArchive& Archive, const FLocalFileReplayInfo& ReplayInfo);
 	LOCALFILENETWORKREPLAYSTREAMING_API bool WriteReplayInfo(FArchive& Archive, const FLocalFileReplayInfo& InReplayInfo, struct FLocalFileSerializationInfo& SerializationInfo);
 
