Add support for gameplay code to inspect the relay header before we attempt to load the map within that replay header.
Add Game Custom logic to load the plugins associated with the desired replay.

diff --git a/Engine/Source/Runtime/Engine/Public/ReplayHelper.h b/Engine/Source/Runtime/Engine/Public/ReplayHelper.h
--- a/Engine/Source/Runtime/Engine/Public/ReplayHelper.h
+++ b/Engine/Source/Runtime/Engine/Public/ReplayHelper.h
@@ -562,7 +562,7 @@ private:
 	void NotifyReplayError(UE::Net::TNetResult<EReplayResult>&& Result);
 
 	bool bRecording;
-
+public:
 	static constexpr int32 MAX_DEMO_READ_WRITE_BUFFER = 1024 * 2;
 	static constexpr int32 MAX_DEMO_STRING_SERIALIZATION_SIZE = 16 * 1024 * 1024;
 };
diff --git a/Engine/Source/Runtime/NetworkReplayStreaming/LocalFileNetworkReplayStreaming/Private/LocalFileNetworkReplayStreaming.cpp b/Engine/Source/Runtime/NetworkReplayStreaming/LocalFileNetworkReplayStreaming/Private/LocalFileNetworkReplayStreaming.cpp
--- a/Engine/Source/Runtime/NetworkReplayStreaming/LocalFileNetworkReplayStreaming/Private/LocalFileNetworkReplayStreaming.cpp
+++ b/Engine/Source/Runtime/NetworkReplayStreaming/LocalFileNetworkReplayStreaming/Private/LocalFileNetworkReplayStreaming.cpp
@@ -11,6 +11,7 @@
 #include "Engine/LocalPlayer.h"
 #include "Engine/World.h"
 #include "Online/CoreOnline.h"
+#include "ReplayHelper.h"
 #include "Serialization/ArchiveCountMem.h"
 #include "Serialization/LargeMemoryReader.h"
 #include "Misc/ConfigCacheIni.h"
@@ -853,16 +854,28 @@ void FLocalFileNetworkReplayStreamer::StartStreaming(const FStartStreamingParame
 					else
 					{
 						DownloadHeader(FDownloadHeaderCallback());
-
-						AddDelegateFileRequestToQueue<FStartStreamingCallback, FStartStreamingResult>(EQueuedLocalFileRequestType::StartPlayback, Delegate,
-							[this](TLocalFileRequestCommonData<FStartStreamingResult>& PlaybackRequestData)
+						AddDelegateFileRequestToQueue<FStartStreamingCallback, FStartStreamingResult>(EQueuedLocalFileRequestType::StartPlayback, FStartStreamingCallback(),
+							[this, Delegate](TLocalFileRequestCommonData<FStartStreamingResult>& PlaybackRequestData)
 							{
-								PlaybackRequestData.DelegateResult.bRecording = false;	
-
-								if (CurrentReplayInfo.bIsValid)
-								{
-									PlaybackRequestData.DelegateResult.Result = EStreamingOperationResult::Success;
-								}
+								PostHeaderDownloaded([this, Delegate]()
+									{
+										if (StreamerState == EReplayStreamerState::Idle || bStopStreamingCalled)
+										{
+											// we were cancelled during operation.
+											return;
+										}
+										AddDelegateFileRequestToQueue<FStartStreamingCallback, FStartStreamingResult>(EQueuedLocalFileRequestType::StartPlayback, Delegate,
+											[this](TLocalFileRequestCommonData<FStartStreamingResult>& PlaybackRequestData)
+											{
+												PlaybackRequestData.DelegateResult.bRecording = false;
+
+												if (CurrentReplayInfo.bIsValid)
+												{
+													PlaybackRequestData.DelegateResult.Result = EStreamingOperationResult::Success;
+												}
+											});
+									}
+								);
 							});
 					}
 				}
diff --git a/Engine/Source/Runtime/NetworkReplayStreaming/LocalFileNetworkReplayStreaming/Public/LocalFileNetworkReplayStreaming.h b/Engine/Source/Runtime/NetworkReplayStreaming/LocalFileNetworkReplayStreaming/Public/LocalFileNetworkReplayStreaming.h
--- a/Engine/Source/Runtime/NetworkReplayStreaming/LocalFileNetworkReplayStreaming/Public/LocalFileNetworkReplayStreaming.h
+++ b/Engine/Source/Runtime/NetworkReplayStreaming/LocalFileNetworkReplayStreaming/Public/LocalFileNetworkReplayStreaming.h
@@ -647,6 +647,8 @@ public:
 
 	static LOCALFILENETWORKREPLAYSTREAMING_API FOnLocalFileReplayFinishedWriting OnReplayFinishedWriting;
 
+	TFunction<void(TFunction<void()> Complete)> PostHeaderDownloaded = [](TFunction<void()> Complete) { Complete(); };
+
 protected:
 
 	LOCALFILENETWORKREPLAYSTREAMING_API void DeleteFinishedStream_Internal(const FString& StreamName, const int32 UserIndex, const FDeleteFinishedStreamCallback& Delegate);
